Made it possible for asset editors to maintain their own FEditorModeTools lists

Breaking changes include:
    * Rename of GEditorModeTools -> GLevelEditorModeTools to signify that it applies only to the level editor modes
    * Addition of FEditorModeRegistry, responsible for managing and creating new editor modes. Modes are no longer registered with an instance of the mode, instead with a mode factory that is able to create a new mode of that type.
    * Editor modes now operate on FEditorViewportClients rather than FLevelEditorViewportClients
    * Added ability to specify an FEditorModeTools when creating an FEditorViewport

Moved component vizualiser manager handling outside of individual editor modes, and into FLevelEditorViewportClient. This should make it easier to transplant in future.

This work addresses TTP#334640 - EDITOR: Investigate making editor modes a per-'editor' concept

diff --git a/Engine/Plugins/Editor/IntroTutorials/Source/IntroTutorials/Private/InteractiveTutorials.cpp b/Engine/Plugins/Editor/IntroTutorials/Source/IntroTutorials/Private/InteractiveTutorials.cpp
--- a/Engine/Plugins/Editor/IntroTutorials/Source/IntroTutorials/Private/InteractiveTutorials.cpp
+++ b/Engine/Plugins/Editor/IntroTutorials/Source/IntroTutorials/Private/InteractiveTutorials.cpp
@@ -34,8 +34,8 @@ void FInteractiveTutorials::SetupEditorHooks()
 
 	FAssetEditorManager::Get().OnAssetEditorRequestedOpen().AddSP(this, &FInteractiveTutorials::OnDelegateTriggered<FAssetEditorRequestedOpenTrigger, FBaseTriggerDelegate::AssetEditorRequestedOpen>);
 
-	GEditorModeTools().OnEditorModeChanged().AddSP( this, &FInteractiveTutorials::OnDelegateTriggered<FEditorModeChangedTrigger, FBaseTriggerDelegate::EditorModeChanged> );
-	GEditorModeTools().OnWidgetModeChanged().AddSP( this, &FInteractiveTutorials::OnDelegateTriggered<FOnWidgetModeChanged, FBaseTriggerDelegate::OnWidgetModeChanged> );
+	GLevelEditorModeTools().OnEditorModeChanged().AddSP( this, &FInteractiveTutorials::OnDelegateTriggered<FEditorModeChangedTrigger, FBaseTriggerDelegate::EditorModeChanged> );
+	GLevelEditorModeTools().OnWidgetModeChanged().AddSP( this, &FInteractiveTutorials::OnDelegateTriggered<FOnWidgetModeChanged, FBaseTriggerDelegate::OnWidgetModeChanged> );
 
 	GEditor->OnEndObjectMovement().AddSP( this, &FInteractiveTutorials::OnDelegateTriggered<FOnEndTransformObject, FBaseTriggerDelegate::OnEndTransformObject> );
 
@@ -61,9 +61,7 @@ void FInteractiveTutorials::SetupEditorHooks()
 	ContentBrowserModule.GetOnSourcesViewChanged().AddSP(this, &FInteractiveTutorials::OnDelegateTriggered<FOnCBSourcesViewChanged, FBaseTriggerDelegate::OnCBSourcesViewChanged>);
 	ContentBrowserModule.GetOnAssetPathChanged().AddSP(this, &FInteractiveTutorials::OnDelegateTriggered<FOnCBAssetPathChanged, FBaseTriggerDelegate::OnCBAssetPathChanged>);
 
-
-	TSharedRef<IPlacementMode> PlacementEditMode = IPlacementModeModule::Get().GetPlacementMode();
-	PlacementEditMode->OnStartedPlacing().AddSP(this, &FInteractiveTutorials::OnDelegateTriggered<FOnStartedPlacing, FBaseTriggerDelegate::OnStartedPlacing>);
+	IPlacementModeModule::Get().OnStartedPlacing().AddSP(this, &FInteractiveTutorials::OnDelegateTriggered<FOnStartedPlacing, FBaseTriggerDelegate::OnStartedPlacing>);
 
 	UK2Node_TutorialExcerptComplete::OnTutorialExcerptComplete.AddSP(this, &FInteractiveTutorials::CompleteExcerpt);
 
diff --git a/Engine/Plugins/Editor/IntroTutorials/Source/IntroTutorials/Private/IntroTutorials.cpp b/Engine/Plugins/Editor/IntroTutorials/Source/IntroTutorials/Private/IntroTutorials.cpp
--- a/Engine/Plugins/Editor/IntroTutorials/Source/IntroTutorials/Private/IntroTutorials.cpp
+++ b/Engine/Plugins/Editor/IntroTutorials/Source/IntroTutorials/Private/IntroTutorials.cpp
@@ -111,7 +111,7 @@ void FIntroTutorials::StartupModule()
 		FGameProjectGenerationModule::Get().OnAddCodeToProjectDialogOpened().AddRaw(this, &FIntroTutorials::OnAddCodeToProjectDialogOpened);
 
 		// Add hook for when editor changes modes (e.g. Place/Paint/Landscape/Foliage)
-		GEditorModeTools().OnEditorModeChanged().AddRaw(this, &FIntroTutorials::OnEditorModeChanged);
+		GLevelEditorModeTools().OnEditorModeChanged().AddRaw(this, &FIntroTutorials::OnEditorModeChanged);
 
 		FSourceCodeNavigation::AccessOnCompilerNotFound().AddRaw( this, &FIntroTutorials::HandleCompilerNotFound );
 
@@ -183,7 +183,7 @@ void FIntroTutorials::ShutdownModule()
 {
 	FSourceCodeNavigation::AccessOnCompilerNotFound().RemoveAll( this );
 
-	GEditorModeTools().OnEditorModeChanged().RemoveAll(this);
+	GLevelEditorModeTools().OnEditorModeChanged().RemoveAll(this);
 
 	if (BlueprintEditorExtender.IsValid() && FModuleManager::Get().IsModuleLoaded("Kismet"))
 	{
diff --git a/Engine/Plugins/Experimental/Paper2D/Source/Paper2DEditor/Private/Paper2DEditorModule.cpp b/Engine/Plugins/Experimental/Paper2D/Source/Paper2DEditor/Private/Paper2DEditorModule.cpp
--- a/Engine/Plugins/Experimental/Paper2D/Source/Paper2DEditor/Private/Paper2DEditorModule.cpp
+++ b/Engine/Plugins/Experimental/Paper2D/Source/Paper2DEditor/Private/Paper2DEditorModule.cpp
@@ -57,8 +57,6 @@ private:
 	/** All created asset type actions.  Cached here so that we can unregister them during shutdown. */
 	TArray< TSharedPtr<IAssetTypeActions> > CreatedAssetTypeActions;
 
-	TSharedPtr<FEdModeTileMap> TileMapEditorModePtr;
-
 	TSharedPtr<IComponentAssetBroker> PaperSpriteBroker;
 	TSharedPtr<IComponentAssetBroker> PaperFlipbookBroker;
 
@@ -114,8 +112,11 @@ public:
 		UThumbnailManager::Get().RegisterCustomRenderer(UPaperFlipbook::StaticClass(), UPaperFlipbookThumbnailRenderer::StaticClass());
 
 		// Register the editor modes
-		TileMapEditorModePtr = MakeShareable(new FEdModeTileMap);
-		GEditorModeTools().RegisterMode(TileMapEditorModePtr.ToSharedRef());
+		FEditorModeRegistry::Get().RegisterMode<FEdModeTileMap>(
+			FEdModeTileMap::EM_TileMap,
+			LOCTEXT("TileMapEditMode", "Tile Map Editor"),
+			FSlateIcon(),
+			true);
 
 		// Integrate Paper2D actions associated with existing engine types (e.g., Texture2D) into the content browser
 		FPaperContentBrowserExtensions::InstallHooks();
@@ -141,7 +142,7 @@ public:
 			FComponentAssetBrokerage::UnregisterBroker(PaperSpriteBroker);
 
 			// Unregister the editor modes
-			GEditorModeTools().UnregisterMode(TileMapEditorModePtr.ToSharedRef());
+			FEditorModeRegistry::Get().UnregisterMode(FEdModeTileMap::EM_TileMap);
 
 			// Unregister the thumbnail renderers
 			UThumbnailManager::Get().UnregisterCustomRenderer(UPaperSprite::StaticClass());
diff --git a/Engine/Plugins/Experimental/Paper2D/Source/Paper2DEditor/Private/PaperEditorViewportClient.cpp b/Engine/Plugins/Experimental/Paper2D/Source/Paper2DEditor/Private/PaperEditorViewportClient.cpp
--- a/Engine/Plugins/Experimental/Paper2D/Source/Paper2DEditor/Private/PaperEditorViewportClient.cpp
+++ b/Engine/Plugins/Experimental/Paper2D/Source/Paper2DEditor/Private/PaperEditorViewportClient.cpp
@@ -8,7 +8,8 @@
 // FPaperEditorViewportClient
 
 FPaperEditorViewportClient::FPaperEditorViewportClient()
-	: CheckerboardTexture(NULL)
+	: FEditorViewportClient(GLevelEditorModeTools())
+	, CheckerboardTexture(NULL)
 {
 	ZoomPos = FVector2D::ZeroVector;
 	ZoomAmount = 1.0f;
diff --git a/Engine/Plugins/Experimental/Paper2D/Source/Paper2DEditor/Private/TileMapEditing/EdModeTileMap.cpp b/Engine/Plugins/Experimental/Paper2D/Source/Paper2DEditor/Private/TileMapEditing/EdModeTileMap.cpp
--- a/Engine/Plugins/Experimental/Paper2D/Source/Paper2DEditor/Private/TileMapEditing/EdModeTileMap.cpp
+++ b/Engine/Plugins/Experimental/Paper2D/Source/Paper2DEditor/Private/TileMapEditing/EdModeTileMap.cpp
@@ -83,10 +83,6 @@ FEdModeTileMap::FEdModeTileMap()
 	, DrawPreviewDimensionsLS(0.0f, 0.0f, 0.0f)
 	, EraseBrushSize(1)
 {
-	ID = EM_TileMap;
-	Name = LOCTEXT("TileMapEditMode", "Tile Map Editor");
-	bVisible = true;
-
 	SetActiveTool(ETileMapEditorTool::Paintbrush);
 	SetActiveLayerPaintingMode(ETileMapLayerPaintingMode::VisualLayers);
 }
@@ -127,7 +123,7 @@ void FEdModeTileMap::Exit()
 	FEdMode::Exit();
 }
 
-bool FEdModeTileMap::MouseMove(FLevelEditorViewportClient* InViewportClient, FViewport* InViewport, int32 x, int32 y)
+bool FEdModeTileMap::MouseMove(FEditorViewportClient* InViewportClient, FViewport* InViewport, int32 x, int32 y)
 {
 	if (InViewportClient->EngineShowFlags.ModeWidgets)
 	{
@@ -139,7 +135,7 @@ bool FEdModeTileMap::MouseMove(FLevelEditorViewportClient* InViewportClient, FVi
 	return false;
 }
 
-bool FEdModeTileMap::CapturedMouseMove(FLevelEditorViewportClient* InViewportClient, FViewport* InViewport, int32 InMouseX, int32 InMouseY)
+bool FEdModeTileMap::CapturedMouseMove(FEditorViewportClient* InViewportClient, FViewport* InViewport, int32 InMouseX, int32 InMouseY)
 {
 	if (InViewportClient->EngineShowFlags.ModeWidgets)
 	{
@@ -161,17 +157,17 @@ bool FEdModeTileMap::CapturedMouseMove(FLevelEditorViewportClient* InViewportCli
 	return false;
 }
 
-bool FEdModeTileMap::StartTracking(FLevelEditorViewportClient* InViewportClient, FViewport* InViewport)
+bool FEdModeTileMap::StartTracking(FEditorViewportClient* InViewportClient, FViewport* InViewport)
 {
 	return true;
 }
 
-bool FEdModeTileMap::EndTracking(FLevelEditorViewportClient* InViewportClient, FViewport* InViewport)
+bool FEdModeTileMap::EndTracking(FEditorViewportClient* InViewportClient, FViewport* InViewport)
 {
 	return true;
 }
 
-bool FEdModeTileMap::InputKey(FLevelEditorViewportClient* InViewportClient, FViewport* InViewport, FKey InKey, EInputEvent InEvent)
+bool FEdModeTileMap::InputKey(FEditorViewportClient* InViewportClient, FViewport* InViewport, FKey InKey, EInputEvent InEvent)
 {
 	bool bHandled = false;
 
@@ -242,7 +238,7 @@ bool FEdModeTileMap::InputKey(FLevelEditorViewportClient* InViewportClient, FVie
 	return bHandled;
 }
 
-bool FEdModeTileMap::InputDelta(FLevelEditorViewportClient* InViewportClient, FViewport* InViewport, FVector& InDrag, FRotator& InRot, FVector& InScale)
+bool FEdModeTileMap::InputDelta(FEditorViewportClient* InViewportClient, FViewport* InViewport, FVector& InDrag, FRotator& InRot, FVector& InScale)
 {
 	return false;
 }
@@ -254,7 +250,7 @@ void FEdModeTileMap::Render(const FSceneView* View, FViewport* Viewport, FPrimit
 	//@TODO: Need the force-realtime hack
 
 	// If this viewport does not support Mode widgets we will not draw it here.
-	FLevelEditorViewportClient* ViewportClient = (FLevelEditorViewportClient*)Viewport->GetClient();
+	FEditorViewportClient* ViewportClient = (FEditorViewportClient*)Viewport->GetClient();
 	if ((ViewportClient != NULL) && !ViewportClient->EngineShowFlags.ModeWidgets)
 	{
 		return;
@@ -280,7 +276,7 @@ void FEdModeTileMap::Render(const FSceneView* View, FViewport* Viewport, FPrimit
 	}
 }
 
-void FEdModeTileMap::DrawHUD(FLevelEditorViewportClient* ViewportClient, FViewport* Viewport, const FSceneView* View, FCanvas* Canvas)
+void FEdModeTileMap::DrawHUD(FEditorViewportClient* ViewportClient, FViewport* Viewport, const FSceneView* View, FCanvas* Canvas)
 {
 	FString InkInfo = FString::Printf(TEXT("Ink: (%d, %d)  %dx%d  %s"), PaintSourceTopLeft.X, PaintSourceTopLeft.Y, PaintSourceDimensions.X, PaintSourceDimensions.Y, 
 		(PaintSourceTileSet.Get() != NULL) ? (*PaintSourceTileSet.Get()->GetName()) : TEXT("(null)"));
@@ -733,7 +729,7 @@ void FEdModeTileMap::UpdatePreviewCursor(const FViewportCursorLocation& Ray)
 	}
 }
 
-FViewportCursorLocation FEdModeTileMap::CalculateViewRay(FLevelEditorViewportClient* InViewportClient, FViewport* InViewport)
+FViewportCursorLocation FEdModeTileMap::CalculateViewRay(FEditorViewportClient* InViewportClient, FViewport* InViewport)
 {
 	FSceneViewFamilyContext ViewFamily( FSceneViewFamily::ConstructionValues( 
 		InViewportClient->Viewport, 
@@ -742,29 +738,11 @@ FViewportCursorLocation FEdModeTileMap::CalculateViewRay(FLevelEditorViewportCli
 		.SetRealtimeUpdate( InViewportClient->IsRealtime() ));
 
 	FSceneView* View = InViewportClient->CalcSceneView( &ViewFamily );
-	FViewportCursorLocation MouseViewportRay( View, (FLevelEditorViewportClient*)InViewport->GetClient(), InViewport->GetMouseX(), InViewport->GetMouseY() );
+	FViewportCursorLocation MouseViewportRay( View, (FEditorViewportClient*)InViewport->GetClient(), InViewport->GetMouseX(), InViewport->GetMouseY() );
 
 	return MouseViewportRay;
 }
 
-void FEdModeTileMap::PeekAtSelectionChangedEvent(UObject* ItemUndergoingChange)
-{
-	/*
-	// Auto-switch to/from tile map editing mode
-	const bool bModeCurrentlyActive = GEditorModeTools().IsModeActive(EM_TileMap);
-	const AActor* TileMapActor = GetFirstSelectedActorContainingTileMapComponent();
-
-	if ((TileMapActor == NULL) && bModeCurrentlyActive)
-	{
-		GEditorModeTools().DeactivateMode(EM_TileMap);
-	}
-	else if ((TileMapActor != NULL) && !bModeCurrentlyActive)
-	{
-		GEditorModeTools().ActivateMode(EM_TileMap);
-	}
-	*/
-}
-
 AActor* FEdModeTileMap::GetFirstSelectedActorContainingTileMapComponent()
 {
 	for (FSelectionIterator It(GEditor->GetSelectedActorIterator()); It; ++It)
diff --git a/Engine/Plugins/Experimental/Paper2D/Source/Paper2DEditor/Private/TileMapEditing/EdModeTileMap.h b/Engine/Plugins/Experimental/Paper2D/Source/Paper2DEditor/Private/TileMapEditing/EdModeTileMap.h
--- a/Engine/Plugins/Experimental/Paper2D/Source/Paper2DEditor/Private/TileMapEditing/EdModeTileMap.h
+++ b/Engine/Plugins/Experimental/Paper2D/Source/Paper2DEditor/Private/TileMapEditing/EdModeTileMap.h
@@ -44,15 +44,15 @@ public:
 	//virtual bool BoxSelect(FBox& InBox, bool bInSelect) override;
 
 	//	virtual void PostUndo() override;
-	virtual bool MouseMove(FLevelEditorViewportClient* ViewportClient, FViewport* Viewport, int32 x, int32 y) override;
-	virtual bool CapturedMouseMove(FLevelEditorViewportClient* InViewportClient, FViewport* InViewport, int32 InMouseX, int32 InMouseY) override;
-	virtual bool StartTracking(FLevelEditorViewportClient* InViewportClient, FViewport* InViewport) override;
-	virtual bool EndTracking(FLevelEditorViewportClient* InViewportClient, FViewport* InViewport) override;
-	// 	virtual void Tick(FLevelEditorViewportClient* ViewportClient,float DeltaTime) override;
-	virtual bool InputKey(FLevelEditorViewportClient* InViewportClient, FViewport* InViewport, FKey InKey, EInputEvent InEvent) override;
-	virtual bool InputDelta(FLevelEditorViewportClient* InViewportClient, FViewport* InViewport, FVector& InDrag, FRotator& InRot, FVector& InScale) override;
+	virtual bool MouseMove(FEditorViewportClient* ViewportClient, FViewport* Viewport, int32 x, int32 y) override;
+	virtual bool CapturedMouseMove(FEditorViewportClient* InViewportClient, FViewport* InViewport, int32 InMouseX, int32 InMouseY) override;
+	virtual bool StartTracking(FEditorViewportClient* InViewportClient, FViewport* InViewport) override;
+	virtual bool EndTracking(FEditorViewportClient* InViewportClient, FViewport* InViewport) override;
+	// 	virtual void Tick(FEditorViewportClient* ViewportClient,float DeltaTime) override;
+	virtual bool InputKey(FEditorViewportClient* InViewportClient, FViewport* InViewport, FKey InKey, EInputEvent InEvent) override;
+	virtual bool InputDelta(FEditorViewportClient* InViewportClient, FViewport* InViewport, FVector& InDrag, FRotator& InRot, FVector& InScale) override;
 	virtual void Render(const FSceneView* View, FViewport* Viewport, FPrimitiveDrawInterface* PDI) override;
-	virtual void DrawHUD(FLevelEditorViewportClient* ViewportClient, FViewport* Viewport, const FSceneView* View, FCanvas* Canvas) override;
+	virtual void DrawHUD(FEditorViewportClient* ViewportClient, FViewport* Viewport, const FSceneView* View, FCanvas* Canvas) override;
 	// 	virtual bool Select(AActor* InActor, bool bInSelected) override;
 	// 	virtual bool IsSelectionAllowed(AActor* InActor) const override;
 	// 	virtual void ActorSelectionChangeNotify() override;
@@ -60,7 +60,6 @@ public:
 	virtual bool AllowWidgetMove();
 	virtual bool ShouldDrawWidget() const override;
 	virtual bool UsesTransformWidget() const override;
-	virtual void PeekAtSelectionChangedEvent(UObject* ItemUndergoingChange) override;
 	// 	virtual int32 GetWidgetAxisToDraw(FWidget::EWidgetMode InWidgetMode) const override;
 	// 	virtual bool DisallowMouseDeltaTracking() const override;
 	// End of FEdMode interface
@@ -89,7 +88,7 @@ protected:
 	UPaperTileLayer* GetSelectedLayerUnderCursor(const FViewportCursorLocation& Ray, int32& OutTileX, int32& OutTileY) const;
 
 	// Compute a world space ray from the screen space mouse coordinates
-	FViewportCursorLocation CalculateViewRay(FLevelEditorViewportClient* InViewportClient, FViewport* InViewport);
+	FViewportCursorLocation CalculateViewRay(FEditorViewportClient* InViewportClient, FViewport* InViewport);
 
 	static AActor* GetFirstSelectedActorContainingTileMapComponent();
 
diff --git a/Engine/Plugins/Experimental/Paper2D/Source/Paper2DEditor/Private/TileMapEditing/PaperTileMapDetailsCustomization.cpp b/Engine/Plugins/Experimental/Paper2D/Source/Paper2DEditor/Private/TileMapEditing/PaperTileMapDetailsCustomization.cpp
--- a/Engine/Plugins/Experimental/Paper2D/Source/Paper2DEditor/Private/TileMapEditing/PaperTileMapDetailsCustomization.cpp
+++ b/Engine/Plugins/Experimental/Paper2D/Source/Paper2DEditor/Private/TileMapEditing/PaperTileMapDetailsCustomization.cpp
@@ -111,7 +111,7 @@ void FPaperTileMapDetailsCustomization::CustomizeDetails(IDetailLayoutBuilder& D
 
 FReply FPaperTileMapDetailsCustomization::EnterTileMapEditingMode()
 {
-	GEditorModeTools().ActivateMode(FEdModeTileMap::EM_TileMap);
+	GLevelEditorModeTools().ActivateMode(FEdModeTileMap::EM_TileMap);
 	return FReply::Handled();
 }
 
@@ -131,7 +131,7 @@ FReply FPaperTileMapDetailsCustomization::AddCollisionLayerClicked()
 
 EVisibility FPaperTileMapDetailsCustomization::GetNonEditModeVisibility() const
 {
-	return GEditorModeTools().IsModeActive(FEdModeTileMap::EM_TileMap) ? EVisibility::Collapsed : EVisibility::Visible;
+	return GLevelEditorModeTools().IsModeActive(FEdModeTileMap::EM_TileMap) ? EVisibility::Collapsed : EVisibility::Visible;
 }
 
 UPaperTileLayer* FPaperTileMapDetailsCustomization::AddLayer(bool bCollisionLayer)
diff --git a/Engine/Plugins/Experimental/Paper2D/Source/Paper2DEditor/Private/TileMapEditing/TileMapEdModeToolkit.cpp b/Engine/Plugins/Experimental/Paper2D/Source/Paper2DEditor/Private/TileMapEditing/TileMapEdModeToolkit.cpp
--- a/Engine/Plugins/Experimental/Paper2D/Source/Paper2DEditor/Private/TileMapEditing/TileMapEdModeToolkit.cpp
+++ b/Engine/Plugins/Experimental/Paper2D/Source/Paper2DEditor/Private/TileMapEditing/TileMapEdModeToolkit.cpp
@@ -46,7 +46,7 @@ FText FTileMapEdModeToolkit::GetToolkitName() const
 
 FEdMode* FTileMapEdModeToolkit::GetEditorMode() const
 {
-	return GEditorModeTools().GetActiveMode(FEdModeTileMap::EM_TileMap);
+	return GLevelEditorModeTools().GetActiveMode(FEdModeTileMap::EM_TileMap);
 }
 
 TSharedPtr<SWidget> FTileMapEdModeToolkit::GetInlineContent() const
@@ -195,7 +195,7 @@ void FTileMapEdModeToolkit::BindCommands()
 
 void FTileMapEdModeToolkit::OnSelectTool(ETileMapEditorTool::Type NewTool)
 {
-	if (FEdModeTileMap* TileMapEditor = GEditorModeTools().GetActiveModeTyped<FEdModeTileMap>(FEdModeTileMap::EM_TileMap))
+	if (FEdModeTileMap* TileMapEditor = GLevelEditorModeTools().GetActiveModeTyped<FEdModeTileMap>(FEdModeTileMap::EM_TileMap))
 	{
 		TileMapEditor->SetActiveTool(NewTool);
 	}
@@ -203,7 +203,7 @@ void FTileMapEdModeToolkit::OnSelectTool(ETileMapEditorTool::Type NewTool)
 
 bool FTileMapEdModeToolkit::IsToolSelected(ETileMapEditorTool::Type QueryTool) const
 {
-	if (FEdModeTileMap* TileMapEditor = GEditorModeTools().GetActiveModeTyped<FEdModeTileMap>(FEdModeTileMap::EM_TileMap))
+	if (FEdModeTileMap* TileMapEditor = GLevelEditorModeTools().GetActiveModeTyped<FEdModeTileMap>(FEdModeTileMap::EM_TileMap))
 	{
 		return (TileMapEditor->GetActiveTool() == QueryTool);
 	}
@@ -215,7 +215,7 @@ bool FTileMapEdModeToolkit::IsToolSelected(ETileMapEditorTool::Type QueryTool) c
 
 void FTileMapEdModeToolkit::OnSelectLayerPaintingMode(ETileMapLayerPaintingMode::Type NewMode)
 {
-	if (FEdModeTileMap* TileMapEditor = GEditorModeTools().GetActiveModeTyped<FEdModeTileMap>(FEdModeTileMap::EM_TileMap))
+	if (FEdModeTileMap* TileMapEditor = GLevelEditorModeTools().GetActiveModeTyped<FEdModeTileMap>(FEdModeTileMap::EM_TileMap))
 	{
 		TileMapEditor->SetActiveLayerPaintingMode(NewMode);
 	}
@@ -223,7 +223,7 @@ void FTileMapEdModeToolkit::OnSelectLayerPaintingMode(ETileMapLayerPaintingMode:
 
 bool FTileMapEdModeToolkit::IsLayerPaintingModeSelected(ETileMapLayerPaintingMode::Type PaintingMode) const
 {
-	if (FEdModeTileMap* TileMapEditor = GEditorModeTools().GetActiveModeTyped<FEdModeTileMap>(FEdModeTileMap::EM_TileMap))
+	if (FEdModeTileMap* TileMapEditor = GLevelEditorModeTools().GetActiveModeTyped<FEdModeTileMap>(FEdModeTileMap::EM_TileMap))
 	{
 		return (TileMapEditor->GetActiveLayerPaintingMode() == PaintingMode);
 	}
@@ -236,7 +236,7 @@ bool FTileMapEdModeToolkit::IsLayerPaintingModeSelected(ETileMapLayerPaintingMod
 EVisibility FTileMapEdModeToolkit::GetTileSetSelectorVisibility() const
 {
 	bool bShouldShowSelector = false;
-	if (FEdModeTileMap* TileMapEditor = GEditorModeTools().GetActiveModeTyped<FEdModeTileMap>(FEdModeTileMap::EM_TileMap))
+	if (FEdModeTileMap* TileMapEditor = GLevelEditorModeTools().GetActiveModeTyped<FEdModeTileMap>(FEdModeTileMap::EM_TileMap))
 	{
 		bShouldShowSelector = (TileMapEditor->GetActiveLayerPaintingMode() == ETileMapLayerPaintingMode::VisualLayers);
 	}
diff --git a/Engine/Plugins/Experimental/Paper2D/Source/Paper2DEditor/Private/TileSetEditor.cpp b/Engine/Plugins/Experimental/Paper2D/Source/Paper2DEditor/Private/TileSetEditor.cpp
--- a/Engine/Plugins/Experimental/Paper2D/Source/Paper2DEditor/Private/TileSetEditor.cpp
+++ b/Engine/Plugins/Experimental/Paper2D/Source/Paper2DEditor/Private/TileSetEditor.cpp
@@ -198,7 +198,7 @@ void STileSetSelectorViewport::OnSelectionChanged(FMarqueeOperation Marquee, boo
 	const bool bHasSelection = (SelectionDimensions.X + SelectionDimensions.Y > 0);
 	if (bIsPreview && bHasSelection)
 	{
-		if (FEdModeTileMap* TileMapEditor = GEditorModeTools().GetActiveModeTyped<FEdModeTileMap>(FEdModeTileMap::EM_TileMap))
+		if (FEdModeTileMap* TileMapEditor = GLevelEditorModeTools().GetActiveModeTyped<FEdModeTileMap>(FEdModeTileMap::EM_TileMap))
 		{
 			TileMapEditor->SetActivePaint(TileSetBeingEdited, SelectionTopLeft, SelectionDimensions);
 
diff --git a/Engine/Source/Editor/ActorPickerMode/Private/ActorPickerMode.cpp b/Engine/Source/Editor/ActorPickerMode/Private/ActorPickerMode.cpp
--- a/Engine/Source/Editor/ActorPickerMode/Private/ActorPickerMode.cpp
+++ b/Engine/Source/Editor/ActorPickerMode/Private/ActorPickerMode.cpp
@@ -6,59 +6,35 @@ IMPLEMENT_MODULE( FActorPickerModeModule, ActorPickerMode );
 
 void FActorPickerModeModule::StartupModule()
 {
-	EdModeActorPicker = MakeShareable(new FEdModeActorPicker);
-	GEditorModeTools().RegisterMode(EdModeActorPicker.ToSharedRef());
-
-	// we want to be able to perform this action with all the built-in editor modes
-	GEditorModeTools().RegisterCompatibleModes(EdModeActorPicker->GetID(), FBuiltinEditorModes::EM_Default);
-	GEditorModeTools().RegisterCompatibleModes(EdModeActorPicker->GetID(), FBuiltinEditorModes::EM_Placement);
-	GEditorModeTools().RegisterCompatibleModes(EdModeActorPicker->GetID(), FBuiltinEditorModes::EM_Bsp);
-	GEditorModeTools().RegisterCompatibleModes(EdModeActorPicker->GetID(), FBuiltinEditorModes::EM_Geometry);
-	GEditorModeTools().RegisterCompatibleModes(EdModeActorPicker->GetID(), FBuiltinEditorModes::EM_InterpEdit);
-	GEditorModeTools().RegisterCompatibleModes(EdModeActorPicker->GetID(), FBuiltinEditorModes::EM_Texture);
-	GEditorModeTools().RegisterCompatibleModes(EdModeActorPicker->GetID(), FBuiltinEditorModes::EM_MeshPaint);
-	GEditorModeTools().RegisterCompatibleModes(EdModeActorPicker->GetID(), FBuiltinEditorModes::EM_Landscape);
-	GEditorModeTools().RegisterCompatibleModes(EdModeActorPicker->GetID(), FBuiltinEditorModes::EM_Foliage);
-	GEditorModeTools().RegisterCompatibleModes(EdModeActorPicker->GetID(), FBuiltinEditorModes::EM_Level);
-	GEditorModeTools().RegisterCompatibleModes(EdModeActorPicker->GetID(), FBuiltinEditorModes::EM_StreamingLevel);
-	GEditorModeTools().RegisterCompatibleModes(EdModeActorPicker->GetID(), FBuiltinEditorModes::EM_Physics);
+	FEditorModeRegistry::Get().RegisterMode<FEdModeActorPicker>(FBuiltinEditorModes::EM_ActorPicker);
 }
 
 void FActorPickerModeModule::ShutdownModule()
 {
-	EndActorPickingMode();
-
-	GEditorModeTools().UnregisterCompatibleModes(EdModeActorPicker->GetID(), FBuiltinEditorModes::EM_Default);
-	GEditorModeTools().UnregisterCompatibleModes(EdModeActorPicker->GetID(), FBuiltinEditorModes::EM_Placement);
-	GEditorModeTools().UnregisterCompatibleModes(EdModeActorPicker->GetID(), FBuiltinEditorModes::EM_Bsp);
-	GEditorModeTools().UnregisterCompatibleModes(EdModeActorPicker->GetID(), FBuiltinEditorModes::EM_Geometry);
-	GEditorModeTools().UnregisterCompatibleModes(EdModeActorPicker->GetID(), FBuiltinEditorModes::EM_InterpEdit);
-	GEditorModeTools().UnregisterCompatibleModes(EdModeActorPicker->GetID(), FBuiltinEditorModes::EM_Texture);
-	GEditorModeTools().UnregisterCompatibleModes(EdModeActorPicker->GetID(), FBuiltinEditorModes::EM_MeshPaint);
-	GEditorModeTools().UnregisterCompatibleModes(EdModeActorPicker->GetID(), FBuiltinEditorModes::EM_Landscape);
-	GEditorModeTools().UnregisterCompatibleModes(EdModeActorPicker->GetID(), FBuiltinEditorModes::EM_Foliage);
-	GEditorModeTools().UnregisterCompatibleModes(EdModeActorPicker->GetID(), FBuiltinEditorModes::EM_Level);
-	GEditorModeTools().UnregisterCompatibleModes(EdModeActorPicker->GetID(), FBuiltinEditorModes::EM_StreamingLevel);
-	GEditorModeTools().UnregisterCompatibleModes(EdModeActorPicker->GetID(), FBuiltinEditorModes::EM_Physics);
-	GEditorModeTools().UnregisterMode(EdModeActorPicker.ToSharedRef());
-	EdModeActorPicker = nullptr;
+	FEditorModeRegistry::Get().UnregisterMode(FBuiltinEditorModes::EM_ActorPicker);
 }
 
 void FActorPickerModeModule::BeginActorPickingMode(FOnGetAllowedClasses InOnGetAllowedClasses, FOnShouldFilterActor InOnShouldFilterActor, FOnActorSelected InOnActorSelected)
 {
-	EndActorPickingMode();
-	EdModeActorPicker->BeginMode(InOnGetAllowedClasses, InOnShouldFilterActor, InOnActorSelected);
+	// Activate the mode
+	GLevelEditorModeTools().ActivateMode(FBuiltinEditorModes::EM_ActorPicker);
+
+	// Set the required delegates
+	FEdModeActorPicker* Mode = GLevelEditorModeTools().GetActiveModeTyped<FEdModeActorPicker>(FBuiltinEditorModes::EM_ActorPicker);
+	if (ensure(Mode))
+	{
+		Mode->OnActorSelected = InOnActorSelected;
+		Mode->OnGetAllowedClasses = InOnGetAllowedClasses;
+		Mode->OnShouldFilterActor = InOnShouldFilterActor;
+	}
 }
 
 void FActorPickerModeModule::EndActorPickingMode()
 {
-	if(IsInActorPickingMode())
-	{
-		EdModeActorPicker->EndMode();
-	}
+	GLevelEditorModeTools().DeactivateMode(FBuiltinEditorModes::EM_ActorPicker);
 }
 
 bool FActorPickerModeModule::IsInActorPickingMode() const
 {
-	return GEditorModeTools().IsModeActive(EdModeActorPicker->GetID());
+	return GLevelEditorModeTools().IsModeActive(FBuiltinEditorModes::EM_ActorPicker);
 }
diff --git a/Engine/Source/Editor/ActorPickerMode/Private/EditorModeActorPicker.cpp b/Engine/Source/Editor/ActorPickerMode/Private/EditorModeActorPicker.cpp
--- a/Engine/Source/Editor/ActorPickerMode/Private/EditorModeActorPicker.cpp
+++ b/Engine/Source/Editor/ActorPickerMode/Private/EditorModeActorPicker.cpp
@@ -6,12 +6,21 @@
 
 FEdModeActorPicker::FEdModeActorPicker()
 {
-	ID = FBuiltinEditorModes::EM_ActorPicker;
 	PickState = EPickState::NotOverViewport;
 	HoveredActor.Reset();
 }
-	
-void FEdModeActorPicker::Tick(FLevelEditorViewportClient* ViewportClient, float DeltaTime)
+
+void FEdModeActorPicker::Initialize()
+{
+	CursorDecoratorWindow = SWindow::MakeCursorDecorator();
+	FSlateApplication::Get().AddWindow(CursorDecoratorWindow.ToSharedRef(), true);
+	CursorDecoratorWindow->SetContent(
+		SNew(SToolTip)
+		.Text(this, &FEdModeActorPicker::GetCursorDecoratorText)
+		);
+}
+
+void FEdModeActorPicker::Tick(FEditorViewportClient* ViewportClient, float DeltaTime)
 {
 	if(CursorDecoratorWindow.IsValid())
 	{
@@ -21,21 +30,21 @@ void FEdModeActorPicker::Tick(FLevelEditorViewportClient* ViewportClient, float
 	FEdMode::Tick(ViewportClient, DeltaTime);
 }
 
-bool FEdModeActorPicker::MouseEnter(FLevelEditorViewportClient* ViewportClient, FViewport* Viewport,int32 x, int32 y)
+bool FEdModeActorPicker::MouseEnter(FEditorViewportClient* ViewportClient, FViewport* Viewport,int32 x, int32 y)
 {
 	PickState = EPickState::OverViewport;
 	HoveredActor.Reset();
 	return FEdMode::MouseEnter(ViewportClient, Viewport, x, y);
 }
 
-bool FEdModeActorPicker::MouseLeave(FLevelEditorViewportClient* ViewportClient, FViewport* Viewport)
+bool FEdModeActorPicker::MouseLeave(FEditorViewportClient* ViewportClient, FViewport* Viewport)
 {
 	PickState = EPickState::NotOverViewport;
 	HoveredActor.Reset();
 	return FEdMode::MouseLeave(ViewportClient, Viewport);
 }
 
-bool FEdModeActorPicker::MouseMove(FLevelEditorViewportClient* ViewportClient, FViewport* Viewport, int32 x, int32 y)
+bool FEdModeActorPicker::MouseMove(FEditorViewportClient* ViewportClient, FViewport* Viewport, int32 x, int32 y)
 {
 	if (ViewportClient == GCurrentLevelEditingViewportClient)
 	{
@@ -65,19 +74,19 @@ bool FEdModeActorPicker::MouseMove(FLevelEditorViewportClient* ViewportClient, F
 	return true;
 }
 
-bool FEdModeActorPicker::LostFocus(FLevelEditorViewportClient* ViewportClient, FViewport* Viewport)
+bool FEdModeActorPicker::LostFocus(FEditorViewportClient* ViewportClient, FViewport* Viewport)
 {
 	if (ViewportClient == GCurrentLevelEditingViewportClient)
 	{
 		// Make sure actor picking mode is disabled once the active viewport loses focus
-		EndMode();
+		RequestDeletion();
 		return true;
 	}
 
 	return false;
 }
 
-bool FEdModeActorPicker::InputKey(FLevelEditorViewportClient* ViewportClient, FViewport* Viewport, FKey Key, EInputEvent Event)
+bool FEdModeActorPicker::InputKey(FEditorViewportClient* ViewportClient, FViewport* Viewport, FKey Key, EInputEvent Event)
 {
 	if (ViewportClient == GCurrentLevelEditingViewportClient)
 	{
@@ -93,20 +102,20 @@ bool FEdModeActorPicker::InputKey(FLevelEditorViewportClient* ViewportClient, FV
 				if(IsActorValid(ActorHit->Actor))
 				{
 					OnActorSelected.ExecuteIfBound(ActorHit->Actor);
-					EndMode();
+					RequestDeletion();
 				}
 			}
 			return true;
 		}
 		else if(Key == EKeys::Escape && Event == IE_Pressed)
 		{
-			EndMode();
+			RequestDeletion();
 			return true;
 		}
 	}
 	else
 	{
-		EndMode();
+		RequestDeletion();
 	}
 
 	return false;
@@ -131,30 +140,17 @@ bool FEdModeActorPicker::UsesToolkits() const
 	return false;
 }
 
-
-void FEdModeActorPicker::BeginMode(FOnGetAllowedClasses InOnGetAllowedClasses, FOnShouldFilterActor InOnShouldFilterActor, FOnActorSelected InOnActorSelected )
+bool FEdModeActorPicker::IsCompatibleWith(FEditorModeID OtherModeID) const
 {
-	OnActorSelected = InOnActorSelected;
-	OnGetAllowedClasses = InOnGetAllowedClasses;
-	OnShouldFilterActor = InOnShouldFilterActor;
-	GEditorModeTools().ActivateMode(GetID());
-	CursorDecoratorWindow = SWindow::MakeCursorDecorator();
-	FSlateApplication::Get().AddWindow(CursorDecoratorWindow.ToSharedRef(), true);
-	CursorDecoratorWindow->SetContent(
-		SNew(SToolTip)
-		.Text(this, &FEdModeActorPicker::GetCursorDecoratorText)
-		);
-
-	HoveredActor.Reset();
-	PickState = EPickState::NotOverViewport;
+	// We want to be able to perform this action with all the built-in editor modes
+	return OtherModeID != FBuiltinEditorModes::EM_None;
 }
 
-void FEdModeActorPicker::EndMode()
+void FEdModeActorPicker::Exit()
 {
 	OnActorSelected = FOnActorSelected();
 	OnGetAllowedClasses = FOnGetAllowedClasses();
 	OnShouldFilterActor = FOnShouldFilterActor();
-	GEditorModeTools().DeactivateMode(GetID());
 
 	if (CursorDecoratorWindow.IsValid())
 	{
@@ -164,6 +160,8 @@ void FEdModeActorPicker::EndMode()
 
 	HoveredActor.Reset();
 	PickState = EPickState::NotOverViewport;
+
+	FEdMode::Exit();
 }
 
 FString FEdModeActorPicker::GetCursorDecoratorText() const
diff --git a/Engine/Source/Editor/ActorPickerMode/Private/EditorModeActorPicker.h b/Engine/Source/Editor/ActorPickerMode/Private/EditorModeActorPicker.h
--- a/Engine/Source/Editor/ActorPickerMode/Private/EditorModeActorPicker.h
+++ b/Engine/Source/Editor/ActorPickerMode/Private/EditorModeActorPicker.h
@@ -20,29 +20,22 @@ class FEdModeActorPicker : public FEdMode
 {
 public:
 	FEdModeActorPicker();
-	
-	/** Register the picker editor mode */
-	static void RegisterEditMode();
-
-	/** Unregister the picker editor mode */
-	static void UnregisterEditMode();
 
 	/** Begin FEdMode interface */
-	virtual void Tick(FLevelEditorViewportClient* ViewportClient, float DeltaTime) override;
-	virtual bool MouseEnter(FLevelEditorViewportClient* ViewportClient, FViewport* Viewport,int32 x, int32 y) override;
-	virtual bool MouseLeave(FLevelEditorViewportClient* ViewportClient, FViewport* Viewport) override;
-	virtual bool MouseMove(FLevelEditorViewportClient* ViewportClient, FViewport* Viewport, int32 x, int32 y) override;
-	virtual bool LostFocus(FLevelEditorViewportClient* ViewportClient, FViewport* Viewport) override;
-	virtual bool InputKey(FLevelEditorViewportClient* ViewportClient, FViewport* Viewport, FKey Key, EInputEvent Event) override;
+	virtual void Tick(FEditorViewportClient* ViewportClient, float DeltaTime) override;
+	virtual bool MouseEnter(FEditorViewportClient* ViewportClient, FViewport* Viewport,int32 x, int32 y) override;
+	virtual bool MouseLeave(FEditorViewportClient* ViewportClient, FViewport* Viewport) override;
+	virtual bool MouseMove(FEditorViewportClient* ViewportClient, FViewport* Viewport, int32 x, int32 y) override;
+	virtual bool LostFocus(FEditorViewportClient* ViewportClient, FViewport* Viewport) override;
+	virtual bool InputKey(FEditorViewportClient* ViewportClient, FViewport* Viewport, FKey Key, EInputEvent Event) override;
 	virtual bool GetCursor(EMouseCursor::Type& OutCursor) const override;
 	virtual bool UsesToolkits() const override;
+	virtual bool IsCompatibleWith(FEditorModeID OtherModeID) const override;
+	virtual void Initialize() override;
 	/** End FEdMode interface */
 
-	/** Start up the mode, reset state etc. */
-	void BeginMode(FOnGetAllowedClasses InOnGetAllowedClasses, FOnShouldFilterActor InOnShouldFilterActor, FOnActorSelected InOnActorSelected );
-
 	/** End the mode */
-	void EndMode();
+	virtual void Exit() override;
 
 	/** Delegate used to display information about picking near the cursor */
 	FString GetCursorDecoratorText() const;
diff --git a/Engine/Source/Editor/ActorPickerMode/Public/ActorPickerMode.h b/Engine/Source/Editor/ActorPickerMode/Public/ActorPickerMode.h
--- a/Engine/Source/Editor/ActorPickerMode/Public/ActorPickerMode.h
+++ b/Engine/Source/Editor/ActorPickerMode/Public/ActorPickerMode.h
@@ -33,8 +33,4 @@ public:
 
 	/** @return Whether or not actor picking mode is currently active */
 	bool IsInActorPickingMode() const;
-
-private:
-	/** The one and only actor picking mode */
-	TSharedPtr<class FEdModeActorPicker> EdModeActorPicker;
 };
diff --git a/Engine/Source/Editor/BspMode/Private/BspMode.cpp b/Engine/Source/Editor/BspMode/Private/BspMode.cpp
--- a/Engine/Source/Editor/BspMode/Private/BspMode.cpp
+++ b/Engine/Source/Editor/BspMode/Private/BspMode.cpp
@@ -3,62 +3,3 @@
 #include "BspModePrivatePCH.h"
 #include "BspMode.h"
 
-TSharedRef< FBspMode > FBspMode::Create()
-{
-	TSharedRef< FBspMode > BspMode = MakeShareable( new FBspMode() );
-	BspMode->Initialize();
-	return BspMode;
-}
-
-FBspMode::FBspMode()
-	: bNeedsToRegisterCommands(true)
-	, LastActiveModeIdentifier()
-{
-	ID = TEXT("BSP");
-	Name = NSLOCTEXT( "GeometryMode", "DisplayName", "Geometry" );
-	IconBrush = FSlateIcon(FEditorStyle::GetStyleSetName(), "LevelEditor.BspMode", "LevelEditor.BspMode.Small");
-	bVisible = true;
-	PriorityOrder = 100;
-}
-
-FBspMode::~FBspMode()
-{
-	FModuleManager::Get().OnModulesChanged().RemoveAll( this );
-}
-
-void FBspMode::Initialize()
-{
-	bVisible = false;
-	FModuleManager::Get().OnModulesChanged().AddSP( this, &FBspMode::HandleModulesChanged );
-}
-
-bool FBspMode::UsesToolkits() const
-{
-	return false;
-}
-
-
-void FBspMode::HandleModulesChanged( FName ModuleThatChanged, EModuleChangeReason::Type ReasonForChange )
-{
-	static const FName LevelEditorModuleName("LevelEditor");
-	if(ModuleThatChanged == LevelEditorModuleName && ReasonForChange == EModuleChangeReason::ModuleUnloaded)
-	{
-		// The LevelEditor module has been unloaded, so we need to re-add our commands next time
-		bNeedsToRegisterCommands = true;
-	}
-}
-
-void FBspMode::Enter()
-{
-	// Call parent implementation
-	FEdMode::Enter();
-
-
-}
-
-void FBspMode::Exit()
-{
-
-	// Call parent implementation
-	FEdMode::Exit();
-}
diff --git a/Engine/Source/Editor/BspMode/Private/BspMode.h b/Engine/Source/Editor/BspMode/Private/BspMode.h
--- a/Engine/Source/Editor/BspMode/Private/BspMode.h
+++ b/Engine/Source/Editor/BspMode/Private/BspMode.h
@@ -7,29 +7,8 @@
 class FBspMode : public FEdMode
 {
 public:
-	static TSharedRef< FBspMode > Create();
 
-	~FBspMode();
-	virtual bool UsesToolkits() const override;
-
-	virtual void Enter() override;
-	virtual void Exit() override;
+	virtual bool UsesToolkits() const override { return false; }
 
 	virtual bool UsesPropertyWidgets() const override { return true; }
-
-private:
-	FBspMode();
-	void Initialize();
-
-	void CreateBspModeItems( FToolBarBuilder& Builder );
-
-	void EditorModeChanged( FEdMode* Mode, bool IsEntering );
-
-	void HandleModulesChanged( FName ModuleThatChanged, EModuleChangeReason::Type ReasonForChange );
-
-private:
-
-	bool IsActive;
-	bool bNeedsToRegisterCommands;
-	FName LastActiveModeIdentifier;
 };
\ No newline at end of file
diff --git a/Engine/Source/Editor/BspMode/Private/BspModeModule.cpp b/Engine/Source/Editor/BspMode/Private/BspModeModule.cpp
--- a/Engine/Source/Editor/BspMode/Private/BspModeModule.cpp
+++ b/Engine/Source/Editor/BspMode/Private/BspModeModule.cpp
@@ -11,8 +11,13 @@ void FBspModeModule::StartupModule()
 {
 	FBspModeStyle::Initialize();
 
-	BspMode = FBspMode::Create();
-	GEditorModeTools().RegisterMode( BspMode.ToSharedRef() );
+	FEditorModeRegistry::Get().RegisterMode<FBspMode>(
+		FBuiltinEditorModes::EM_Bsp,
+		NSLOCTEXT("GeometryMode", "DisplayName", "Geometry"),
+		FSlateIcon(FEditorStyle::GetStyleSetName(), "LevelEditor.BspMode", "LevelEditor.BspMode.Small"),
+		false,		// Visible
+		100			// UI priority order
+		);
 
 	RegisterBspBuilderType(UCubeBuilder::StaticClass(), LOCTEXT("CubeBuilderName", "Box"), LOCTEXT("CubeBuilderToolTip", "Make a box brush"), FBspModeStyle::Get().GetBrush(TEXT("BspMode.BoxBrush")));
 	RegisterBspBuilderType(UConeBuilder::StaticClass(), LOCTEXT("ConeBuilderName", "Cone"), LOCTEXT("ConeBuilderToolTip", "Make a cone brush"), FBspModeStyle::Get().GetBrush(TEXT("BspMode.ConeBrush")));
@@ -26,8 +31,7 @@ void FBspModeModule::StartupModule()
 
 void FBspModeModule::ShutdownModule()
 {
-	GEditorModeTools().UnregisterMode( BspMode.ToSharedRef() );
-	BspMode = NULL;
+	FEditorModeRegistry::Get().UnregisterMode( FBuiltinEditorModes::EM_Bsp );
 
 	BspBuilderTypes.Empty();
 }
@@ -39,12 +43,6 @@ TSharedRef< SWidget > FBspModeModule::CreateBspModeWidget() const
 }
 
 
-TSharedRef< FEdMode > FBspModeModule::GetBspMode() const
-{
-	return BspMode.ToSharedRef();
-}
-
-
 void FBspModeModule::RegisterBspBuilderType( class UClass* InBuilderClass, const FText& InBuilderName, const FText& InBuilderTooltip, const FSlateBrush* InBuilderIcon )
 {
 	check(InBuilderClass->IsChildOf(UBrushBuilder::StaticClass()));
diff --git a/Engine/Source/Editor/BspMode/Private/BspModeModule.h b/Engine/Source/Editor/BspMode/Private/BspModeModule.h
--- a/Engine/Source/Editor/BspMode/Private/BspModeModule.h
+++ b/Engine/Source/Editor/BspMode/Private/BspModeModule.h
@@ -37,7 +37,6 @@ public:
 
 	/** IBspModeModule interface */
 	virtual TSharedRef< SWidget > CreateBspModeWidget() const override;
-	virtual TSharedRef< FEdMode > GetBspMode() const override;
 	virtual void RegisterBspBuilderType( class UClass* InBuilderClass, const FText& InBuilderName, const FText& InBuilderTooltip, const FSlateBrush* InBuilderIcon ) override;
 	virtual void UnregisterBspBuilderType( class UClass* InBuilderClass ) override;
 
@@ -46,7 +45,5 @@ public:
 
 private:
 
-	TSharedPtr<class FBspMode> BspMode;
-
 	TArray< TSharedPtr<FBspBuilderType> > BspBuilderTypes;
 };
diff --git a/Engine/Source/Editor/BspMode/Public/IBspModeModule.h b/Engine/Source/Editor/BspMode/Public/IBspModeModule.h
--- a/Engine/Source/Editor/BspMode/Public/IBspModeModule.h
+++ b/Engine/Source/Editor/BspMode/Public/IBspModeModule.h
@@ -5,7 +5,6 @@ class IBspModeModule : public IModuleInterface
 {
 public:
 	virtual TSharedRef< SWidget > CreateBspModeWidget() const = 0;
-	virtual TSharedRef< FEdMode > GetBspMode() const = 0;
 
 	virtual void RegisterBspBuilderType( class UClass* InBuilderClass, const FText& InBuilderName, const FText& InBuilderTooltip, const FSlateBrush* InBuilderIcon ) = 0;
 	virtual void UnregisterBspBuilderType( class UClass* InBuilderClass ) = 0;
diff --git a/Engine/Source/Editor/Cascade/Private/CascadeEmitterCanvasClient.cpp b/Engine/Source/Editor/Cascade/Private/CascadeEmitterCanvasClient.cpp
--- a/Engine/Source/Editor/Cascade/Private/CascadeEmitterCanvasClient.cpp
+++ b/Engine/Source/Editor/Cascade/Private/CascadeEmitterCanvasClient.cpp
@@ -19,7 +19,8 @@
 
 
 FCascadeEmitterCanvasClient::FCascadeEmitterCanvasClient(TWeakPtr<FCascade> InCascade, TWeakPtr<SCascadeEmitterCanvas> InCascadeViewport)
-	: CascadePtr(InCascade)
+	: FEditorViewportClient(GLevelEditorModeTools())
+	, CascadePtr(InCascade)
 	, CascadeViewportPtr(InCascadeViewport)
 	, EmitterWidth(180)
 	, EmitterCollapsedWidth(18)
diff --git a/Engine/Source/Editor/Cascade/Private/CascadePreviewViewportClient.cpp b/Engine/Source/Editor/Cascade/Private/CascadePreviewViewportClient.cpp
--- a/Engine/Source/Editor/Cascade/Private/CascadePreviewViewportClient.cpp
+++ b/Engine/Source/Editor/Cascade/Private/CascadePreviewViewportClient.cpp
@@ -17,7 +17,8 @@
 #define LOCTEXT_NAMESPACE "CascadeViewportClient"
 
 FCascadeEdPreviewViewportClient::FCascadeEdPreviewViewportClient(TWeakPtr<FCascade> InCascade, TWeakPtr<SCascadePreviewViewport> InCascadeViewport)
-	: CascadePtr(InCascade)
+	: FEditorViewportClient(GLevelEditorModeTools())
+	, CascadePtr(InCascade)
 	, CascadeViewportPtr(InCascadeViewport)
 	, CascadePreviewScene(FPreviewScene::ConstructionValues()
 						.SetLightRotation(FRotator(-45.f, 180.f, 0.f))
diff --git a/Engine/Source/Editor/ComponentVisualizers/Private/SplineComponentVisualizer.cpp b/Engine/Source/Editor/ComponentVisualizers/Private/SplineComponentVisualizer.cpp
--- a/Engine/Source/Editor/ComponentVisualizers/Private/SplineComponentVisualizer.cpp
+++ b/Engine/Source/Editor/ComponentVisualizers/Private/SplineComponentVisualizer.cpp
@@ -211,7 +211,7 @@ bool FSplineComponentVisualizer::GetWidgetLocation(FVector& OutLocation) const
 }
 
 
-bool FSplineComponentVisualizer::HandleInputDelta(FLevelEditorViewportClient* ViewportClient, FViewport* Viewport, FVector& DeltaTranslate, FRotator& DeltaRotate, FVector& DeltaScale)
+bool FSplineComponentVisualizer::HandleInputDelta(FEditorViewportClient* ViewportClient, FViewport* Viewport, FVector& DeltaTranslate, FRotator& DeltaRotate, FVector& DeltaScale)
 {
 	USplineComponent* SplineComp = GetEditedSplineComponent();
 	if( SplineComp && 
@@ -282,7 +282,7 @@ bool FSplineComponentVisualizer::HandleInputDelta(FLevelEditorViewportClient* Vi
 	return false;
 }
 
-bool FSplineComponentVisualizer::HandleInputKey(FLevelEditorViewportClient* ViewportClient, FViewport* Viewport, FKey Key, EInputEvent Event)
+bool FSplineComponentVisualizer::HandleInputKey(FEditorViewportClient* ViewportClient, FViewport* Viewport, FKey Key, EInputEvent Event)
 {
 	// TODO: find a way to make this work with ProcessCommandBindings(FKeyboardEvent&)
 
diff --git a/Engine/Source/Editor/ComponentVisualizers/Private/SplineComponentVisualizer.h b/Engine/Source/Editor/ComponentVisualizers/Private/SplineComponentVisualizer.h
--- a/Engine/Source/Editor/ComponentVisualizers/Private/SplineComponentVisualizer.h
+++ b/Engine/Source/Editor/ComponentVisualizers/Private/SplineComponentVisualizer.h
@@ -42,8 +42,8 @@ public:
 	virtual bool VisProxyHandleClick(HComponentVisProxy* VisProxy) override;
 	virtual void EndEditing() override;
 	virtual bool GetWidgetLocation(FVector& OutLocation) const override;
-	virtual bool HandleInputDelta(FLevelEditorViewportClient* ViewportClient, FViewport* Viewport, FVector& DeltaTranslate, FRotator& DeltaRotate, FVector& DeltaScale) override;
-	virtual bool HandleInputKey(FLevelEditorViewportClient* ViewportClient, FViewport* Viewport, FKey Key, EInputEvent Event) override;
+	virtual bool HandleInputDelta(FEditorViewportClient* ViewportClient, FViewport* Viewport, FVector& DeltaTranslate, FRotator& DeltaRotate, FVector& DeltaScale) override;
+	virtual bool HandleInputKey(FEditorViewportClient* ViewportClient, FViewport* Viewport, FKey Key, EInputEvent Event) override;
 	virtual TSharedPtr<SWidget> GenerateContextMenu() const override;
 	// End FComponentVisualizer interface
 
diff --git a/Engine/Source/Editor/DestructibleMeshEditor/Private/SDestructibleMeshEditorViewport.cpp b/Engine/Source/Editor/DestructibleMeshEditor/Private/SDestructibleMeshEditorViewport.cpp
--- a/Engine/Source/Editor/DestructibleMeshEditor/Private/SDestructibleMeshEditorViewport.cpp
+++ b/Engine/Source/Editor/DestructibleMeshEditor/Private/SDestructibleMeshEditorViewport.cpp
@@ -77,7 +77,7 @@ private:
 };
 
 FDestructibleMeshEditorViewportClient::FDestructibleMeshEditorViewportClient(TWeakPtr<IDestructibleMeshEditor> InDestructibleMeshEditor, FPreviewScene& InPreviewScene)
-	: FEditorViewportClient(&InPreviewScene)
+	: FEditorViewportClient(GLevelEditorModeTools(), &InPreviewScene)
 	, DestructibleMeshEditorPtr(InDestructibleMeshEditor)
 {
 	SetViewMode(VMI_Lit);
diff --git a/Engine/Source/Editor/DetailCustomizations/Private/ActorDetails.cpp b/Engine/Source/Editor/DetailCustomizations/Private/ActorDetails.cpp
--- a/Engine/Source/Editor/DetailCustomizations/Private/ActorDetails.cpp
+++ b/Engine/Source/Editor/DetailCustomizations/Private/ActorDetails.cpp
@@ -297,7 +297,7 @@ TSharedRef<SWidget> FActorDetails::OnGetConvertContent()
 
 EVisibility FActorDetails::GetConvertMenuVisibility() const
 {
-	return GEditorModeTools().EnsureNotInMode(FBuiltinEditorModes::EM_InterpEdit) ?
+	return GLevelEditorModeTools().EnsureNotInMode(FBuiltinEditorModes::EM_InterpEdit) ?
 		EVisibility::Visible :
 		EVisibility::Collapsed;
 }
diff --git a/Engine/Source/Editor/DetailCustomizations/Private/AnimMontageSegmentDetails.cpp b/Engine/Source/Editor/DetailCustomizations/Private/AnimMontageSegmentDetails.cpp
--- a/Engine/Source/Editor/DetailCustomizations/Private/AnimMontageSegmentDetails.cpp
+++ b/Engine/Source/Editor/DetailCustomizations/Private/AnimMontageSegmentDetails.cpp
@@ -10,7 +10,7 @@
 
 /////////////////////////////////////////////////////////////////////////
 FAnimationSegmentViewportClient::FAnimationSegmentViewportClient(FPreviewScene& InPreviewScene, FFXSystemInterface * FXSystem)
-	: FEditorViewportClient(&InPreviewScene)
+	: FEditorViewportClient(GLevelEditorModeTools(), &InPreviewScene)
 {
 	SetViewMode(VMI_Lit);
 
diff --git a/Engine/Source/Editor/DetailCustomizations/Private/LandscapeSplineDetails.cpp b/Engine/Source/Editor/DetailCustomizations/Private/LandscapeSplineDetails.cpp
--- a/Engine/Source/Editor/DetailCustomizations/Private/LandscapeSplineDetails.cpp
+++ b/Engine/Source/Editor/DetailCustomizations/Private/LandscapeSplineDetails.cpp
@@ -64,7 +64,7 @@ void FLandscapeSplineDetails::CustomizeDetails( IDetailLayoutBuilder& DetailBuil
 
 class FEdModeLandscape* FLandscapeSplineDetails::GetEditorMode() const
 {
-	return (FEdModeLandscape*)GEditorModeTools().GetActiveMode(FBuiltinEditorModes::EM_Landscape);
+	return (FEdModeLandscape*)GLevelEditorModeTools().GetActiveMode(FBuiltinEditorModes::EM_Landscape);
 }
 
 FReply FLandscapeSplineDetails::OnSelectConnectedControlPointsButtonClicked()
diff --git a/Engine/Source/Editor/FoliageEdit/Private/FoliageEdMode.cpp b/Engine/Source/Editor/FoliageEdit/Private/FoliageEdMode.cpp
--- a/Engine/Source/Editor/FoliageEdit/Private/FoliageEdMode.cpp
+++ b/Engine/Source/Editor/FoliageEdit/Private/FoliageEdMode.cpp
@@ -41,12 +41,6 @@ FEdModeFoliage::FEdModeFoliage()
 ,	SelectionIFA(NULL)
 ,	bCanAltDrag(false)
 {
-	ID = FBuiltinEditorModes::EM_Foliage;
-	Name = NSLOCTEXT("EditorModes", "FoliageMode", "Foliage");
-	IconBrush = FSlateIcon(FEditorStyle::GetStyleSetName(), "LevelEditor.FoliageMode", "LevelEditor.FoliageMode.Small");
-	bVisible = true;
-	PriorityOrder = 400;
-
 	// Load resources and construct brush component
 	UMaterial* BrushMaterial = NULL;
 	UStaticMesh* StaticMesh = NULL;
@@ -71,6 +65,7 @@ FEdModeFoliage::~FEdModeFoliage()
 {
 	// Save UI settings to config file
 	UISettings.Save();
+	FEditorDelegates::MapChange.RemoveAll(this);
 }
 
 
@@ -83,6 +78,12 @@ void FEdModeFoliage::AddReferencedObjects( FReferenceCollector& Collector )
 	Collector.AddReferencedObject(SphereBrushComponent);
 }
 
+/** FEdMode: Called when the mode is created */
+void FEdModeFoliage::Initialize()
+{
+	FEditorDelegates::MapChange.AddSP( this, &FEdModeFoliage::NotifyMapRebuild );
+}
+
 /** FEdMode: Called when the mode is entered */
 void FEdModeFoliage::Enter()
 {
@@ -236,7 +237,7 @@ bool FEdModeFoliage::DisallowMouseDeltaTracking() const
 }
 
 /** FEdMode: Called once per frame */
-void FEdModeFoliage::Tick(FLevelEditorViewportClient* ViewportClient,float DeltaTime)
+void FEdModeFoliage::Tick(FEditorViewportClient* ViewportClient,float DeltaTime)
 {
 	if( IsCtrlDown(ViewportClient->Viewport) && bToolActive )
 	{
@@ -249,8 +250,7 @@ void FEdModeFoliage::Tick(FLevelEditorViewportClient* ViewportClient,float Delta
 	{
 		// Update pivot
 		AInstancedFoliageActor* IFA = AInstancedFoliageActor::GetInstancedFoliageActor(GetWorld());
-		FEditorModeTools& Tools = GEditorModeTools();
-		Tools.PivotLocation = Tools.SnappedLocation = IFA->GetSelectionLocation();
+		Owner->PivotLocation = Owner->SnappedLocation = IFA->GetSelectionLocation();
 	}
 	
 	// Update the position and size of the brush component
@@ -275,7 +275,7 @@ void FEdModeFoliage::Tick(FLevelEditorViewportClient* ViewportClient,float Delta
 }
 
 /** Trace under the mouse cursor and update brush position */
-void FEdModeFoliage::FoliageBrushTrace( FLevelEditorViewportClient* ViewportClient, int32 MouseX, int32 MouseY )
+void FEdModeFoliage::FoliageBrushTrace( FEditorViewportClient* ViewportClient, int32 MouseX, int32 MouseY )
 {
 	bBrushTraceValid = false;
 	if( UISettings.GetPaintToolSelected() || UISettings.GetReapplyToolSelected() || UISettings.GetLassoSelectToolSelected() )
@@ -328,7 +328,7 @@ void FEdModeFoliage::FoliageBrushTrace( FLevelEditorViewportClient* ViewportClie
  *
  * @return	true if input was handled
  */
-bool FEdModeFoliage::MouseMove( FLevelEditorViewportClient* ViewportClient, FViewport* Viewport, int32 MouseX, int32 MouseY )
+bool FEdModeFoliage::MouseMove( FEditorViewportClient* ViewportClient, FViewport* Viewport, int32 MouseX, int32 MouseY )
 {
 	FoliageBrushTrace(ViewportClient, MouseX, MouseY );
 	return false;
@@ -344,7 +344,7 @@ bool FEdModeFoliage::MouseMove( FLevelEditorViewportClient* ViewportClient, FVie
  *
  * @return	true if input was handled
  */
-bool FEdModeFoliage::CapturedMouseMove( FLevelEditorViewportClient* ViewportClient, FViewport* Viewport, int32 MouseX, int32 MouseY )
+bool FEdModeFoliage::CapturedMouseMove( FEditorViewportClient* ViewportClient, FViewport* Viewport, int32 MouseX, int32 MouseY )
 {
 	FoliageBrushTrace(ViewportClient, MouseX, MouseY );
 	return false;
@@ -1101,7 +1101,7 @@ void FEdModeFoliage::ReapplyInstancesForBrush( UWorld* InWorld, AInstancedFoliag
 	}
 }
 
-void FEdModeFoliage::ApplyBrush( FLevelEditorViewportClient* ViewportClient )
+void FEdModeFoliage::ApplyBrush( FEditorViewportClient* ViewportClient )
 {
 	if( !bBrushTraceValid )
 	{
@@ -1190,8 +1190,7 @@ void FEdModeFoliage::ApplyBrush( FLevelEditorViewportClient* ViewportClient )
 	if( UISettings.GetLassoSelectToolSelected() )
 	{
 		IFA->CheckSelection();
-		FEditorModeTools& Tools = GEditorModeTools();
-		Tools.PivotLocation = Tools.SnappedLocation = IFA->GetSelectionLocation();
+		Owner->PivotLocation = Owner->SnappedLocation = IFA->GetSelectionLocation();
 		IFA->MarkComponentsRenderStateDirty();
 	}
 }
@@ -1748,7 +1747,7 @@ bool FEdModeFoliage::ReplaceStaticMesh(UStaticMesh* OldStaticMesh, UStaticMesh*
 
 
 /** FEdMode: Called when a key is pressed */
-bool FEdModeFoliage::InputKey( FLevelEditorViewportClient* ViewportClient, FViewport* Viewport, FKey Key, EInputEvent Event )
+bool FEdModeFoliage::InputKey( FEditorViewportClient* ViewportClient, FViewport* Viewport, FKey Key, EInputEvent Event )
 {
 	if( UISettings.GetPaintToolSelected() || UISettings.GetReapplyToolSelected() || UISettings.GetLassoSelectToolSelected() )
 	{
@@ -1875,8 +1874,7 @@ bool FEdModeFoliage::InputKey( FLevelEditorViewportClient* ViewportClient, FView
 
 				if( bMovedInstance )
 				{
-					FEditorModeTools& Tools = GEditorModeTools();
-					Tools.PivotLocation = Tools.SnappedLocation = IFA->GetSelectionLocation();
+					Owner->PivotLocation = Owner->SnappedLocation = IFA->GetSelectionLocation();
 					IFA->MarkComponentsRenderStateDirty();
 				}
 
@@ -1897,7 +1895,7 @@ void FEdModeFoliage::Render( const FSceneView* View, FViewport* Viewport, FPrimi
 
 
 /** FEdMode: Render HUD elements for this tool */
-void FEdModeFoliage::DrawHUD( FLevelEditorViewportClient* ViewportClient, FViewport* Viewport, const FSceneView* View, FCanvas* Canvas )
+void FEdModeFoliage::DrawHUD( FEditorViewportClient* ViewportClient, FViewport* Viewport, const FSceneView* View, FCanvas* Canvas )
 {
 }
 
@@ -1932,7 +1930,7 @@ void FEdModeFoliage::ForceRealTimeViewports( const bool bEnable, const bool bSto
 	TSharedPtr< ILevelViewport > ViewportWindow = LevelEditorModule.GetFirstActiveViewport();
 	if (ViewportWindow.IsValid())
 	{
-		FLevelEditorViewportClient &Viewport = ViewportWindow->GetLevelViewportClient();
+		FEditorViewportClient &Viewport = ViewportWindow->GetLevelViewportClient();
 		if( Viewport.IsPerspective() )
 		{				
 			if( bEnable )
@@ -1949,7 +1947,7 @@ void FEdModeFoliage::ForceRealTimeViewports( const bool bEnable, const bool bSto
 	}
 }
 
-bool FEdModeFoliage::HandleClick(FLevelEditorViewportClient* InViewportClient, HHitProxy *HitProxy, const FViewportClick &Click )
+bool FEdModeFoliage::HandleClick(FEditorViewportClient* InViewportClient, HHitProxy *HitProxy, const FViewportClick &Click )
 {
 	AInstancedFoliageActor* IFA = AInstancedFoliageActor::GetInstancedFoliageActor(GetWorld());
 
@@ -1962,8 +1960,7 @@ bool FEdModeFoliage::HandleClick(FLevelEditorViewportClient* InViewportClient, H
 			IFA->SelectInstance( SMIProxy->Component, SMIProxy->InstanceIndex, Click.IsControlDown() );
 
 			// Update pivot
-			FEditorModeTools& Tools = GEditorModeTools();
-			Tools.PivotLocation = Tools.SnappedLocation = IFA->GetSelectionLocation();
+			Owner->PivotLocation = Owner->SnappedLocation = IFA->GetSelectionLocation();
 		}
 		else
 		{
@@ -1998,14 +1995,13 @@ FVector FEdModeFoliage::GetWidgetLocation() const
 }
 
 /** FEdMode: Called when a mouse button is pressed */
-bool FEdModeFoliage::StartTracking(FLevelEditorViewportClient* InViewportClient, FViewport* InViewport)
+bool FEdModeFoliage::StartTracking(FEditorViewportClient* InViewportClient, FViewport* InViewport)
 {
 	if( UISettings.GetSelectToolSelected() || UISettings.GetLassoSelectToolSelected() )
 	{
 		// Update pivot
 		AInstancedFoliageActor* IFA = AInstancedFoliageActor::GetInstancedFoliageActor(GetWorld());
-		FEditorModeTools& Tools = GEditorModeTools();
-		Tools.PivotLocation = Tools.SnappedLocation = IFA->GetSelectionLocation();
+		Owner->PivotLocation = Owner->SnappedLocation = IFA->GetSelectionLocation();
 
 		GEditor->BeginTransaction( NSLOCTEXT("UnrealEd", "FoliageMode_EditTransaction", "Foliage Editing") );
 		
@@ -2017,7 +2013,7 @@ bool FEdModeFoliage::StartTracking(FLevelEditorViewportClient* InViewportClient,
 }
 
 /** FEdMode: Called when the a mouse button is released */
-bool FEdModeFoliage::EndTracking(FLevelEditorViewportClient* InViewportClient, FViewport* InViewport)
+bool FEdModeFoliage::EndTracking(FEditorViewportClient* InViewportClient, FViewport* InViewport)
 {
 	if( UISettings.GetSelectToolSelected() || UISettings.GetLassoSelectToolSelected() )
 	{
@@ -2028,7 +2024,7 @@ bool FEdModeFoliage::EndTracking(FLevelEditorViewportClient* InViewportClient, F
 }
 
 /** FEdMode: Called when mouse drag input it applied */
-bool FEdModeFoliage::InputDelta( FLevelEditorViewportClient* InViewportClient, FViewport* InViewport, FVector& InDrag, FRotator& InRot, FVector& InScale )
+bool FEdModeFoliage::InputDelta( FEditorViewportClient* InViewportClient, FViewport* InViewport, FVector& InDrag, FRotator& InRot, FVector& InScale )
 {
 	bool bFoundSelection = false;
 	
diff --git a/Engine/Source/Editor/FoliageEdit/Private/FoliageEdMode.h b/Engine/Source/Editor/FoliageEdit/Private/FoliageEdMode.h
--- a/Engine/Source/Editor/FoliageEdit/Private/FoliageEdMode.h
+++ b/Engine/Source/Editor/FoliageEdit/Private/FoliageEdMode.h
@@ -163,6 +163,9 @@ public:
 	/** FGCObject interface */
 	virtual void AddReferencedObjects( FReferenceCollector& Collector ) override;
 
+	/** FEdMode: Called when the mode is created */
+	virtual void Initialize() override;
+
 	/** FEdMode: Called when the mode is entered */
 	virtual void Enter() override;
 
@@ -194,7 +197,7 @@ public:
 	 *
 	 * @return	true if input was handled
 	 */
-	virtual bool MouseMove( FLevelEditorViewportClient* ViewportClient, FViewport* Viewport, int32 x, int32 y ) override;
+	virtual bool MouseMove( FEditorViewportClient* ViewportClient, FViewport* Viewport, int32 x, int32 y ) override;
 
 	/**
 	 * FEdMode: Called when the mouse is moved while a window input capture is in effect
@@ -206,28 +209,28 @@ public:
 	 *
 	 * @return	true if input was handled
 	 */
-	virtual bool CapturedMouseMove( FLevelEditorViewportClient* InViewportClient, FViewport* InViewport, int32 InMouseX, int32 InMouseY ) override;
+	virtual bool CapturedMouseMove( FEditorViewportClient* InViewportClient, FViewport* InViewport, int32 InMouseX, int32 InMouseY ) override;
 
 	/** FEdMode: Called when a mouse button is pressed */
-	virtual bool StartTracking(FLevelEditorViewportClient* InViewportClient, FViewport* InViewport) override;
+	virtual bool StartTracking(FEditorViewportClient* InViewportClient, FViewport* InViewport) override;
 
 	/** FEdMode: Called when a mouse button is released */
-	virtual bool EndTracking(FLevelEditorViewportClient* InViewportClient, FViewport* InViewport) override;
+	virtual bool EndTracking(FEditorViewportClient* InViewportClient, FViewport* InViewport) override;
 
 	/** FEdMode: Called once per frame */
-	virtual void Tick(FLevelEditorViewportClient* ViewportClient,float DeltaTime) override;
+	virtual void Tick(FEditorViewportClient* ViewportClient,float DeltaTime) override;
 
 	/** FEdMode: Called when a key is pressed */
-	virtual bool InputKey( FLevelEditorViewportClient* InViewportClient, FViewport* InViewport, FKey InKey, EInputEvent InEvent ) override;
+	virtual bool InputKey( FEditorViewportClient* InViewportClient, FViewport* InViewport, FKey InKey, EInputEvent InEvent ) override;
 
 	/** FEdMode: Called when mouse drag input it applied */
-	virtual bool InputDelta( FLevelEditorViewportClient* InViewportClient, FViewport* InViewport, FVector& InDrag, FRotator& InRot, FVector& InScale ) override;
+	virtual bool InputDelta( FEditorViewportClient* InViewportClient, FViewport* InViewport, FVector& InDrag, FRotator& InRot, FVector& InScale ) override;
 
 	/** FEdMode: Render elements for the Foliage tool */
 	virtual void Render( const FSceneView* View, FViewport* Viewport, FPrimitiveDrawInterface* PDI ) override;
 
 	/** FEdMode: Render HUD elements for this tool */
-	virtual void DrawHUD( FLevelEditorViewportClient* ViewportClient, FViewport* Viewport, const FSceneView* View, FCanvas* Canvas ) override;
+	virtual void DrawHUD( FEditorViewportClient* ViewportClient, FViewport* Viewport, const FSceneView* View, FCanvas* Canvas ) override;
 
 	/** FEdMode: Handling SelectActor */
 	virtual bool Select( AActor* InActor, bool bInSelected ) override;
@@ -239,7 +242,7 @@ public:
 	virtual void ActorSelectionChangeNotify() override;
 
 	/** Notifies all active modes of mouse click messages. */
-	bool HandleClick(FLevelEditorViewportClient* InViewportClient,  HHitProxy *HitProxy, const FViewportClick &Click );
+	bool HandleClick(FEditorViewportClient* InViewportClient,  HHitProxy *HitProxy, const FViewportClick &Click );
 
 	/** FEdMode: widget handling */
 	virtual FVector GetWidgetLocation() const override;
@@ -254,14 +257,14 @@ public:
 	void ForceRealTimeViewports( const bool bEnable, const bool bStoreCurrentState );
 
 	/** Trace under the mouse cursor and update brush position */
-	void FoliageBrushTrace( FLevelEditorViewportClient* ViewportClient, int32 MouseX, int32 MouseY );
+	void FoliageBrushTrace( FEditorViewportClient* ViewportClient, int32 MouseX, int32 MouseY );
 
 	/** Generate start/end points for a random trace inside the sphere brush. 
 	    returns a line segment inside the sphere parallel to the view direction */
 	void GetRandomVectorInBrush(FVector& OutStart, FVector& OutEnd );
 
 	/** Apply brush */
-	void ApplyBrush( FLevelEditorViewportClient* ViewportClient );
+	void ApplyBrush( FEditorViewportClient* ViewportClient );
 
 	/** Update existing mesh info for current level */
 	void UpdateFoliageMeshList();
diff --git a/Engine/Source/Editor/FoliageEdit/Private/FoliageEdModeToolkit.cpp b/Engine/Source/Editor/FoliageEdit/Private/FoliageEdModeToolkit.cpp
--- a/Engine/Source/Editor/FoliageEdit/Private/FoliageEdModeToolkit.cpp
+++ b/Engine/Source/Editor/FoliageEdit/Private/FoliageEdModeToolkit.cpp
@@ -36,7 +36,7 @@ FText FFoliageEdModeToolkit::GetBaseToolkitName() const
 
 class FEdMode* FFoliageEdModeToolkit::GetEditorMode() const
 {
-	return GEditorModeTools().GetActiveMode(FBuiltinEditorModes::EM_Foliage);
+	return GLevelEditorModeTools().GetActiveMode(FBuiltinEditorModes::EM_Foliage);
 }
 
 TSharedPtr<SWidget> FFoliageEdModeToolkit::GetInlineContent() const
diff --git a/Engine/Source/Editor/FoliageEdit/Private/FoliageEditModule.cpp b/Engine/Source/Editor/FoliageEdit/Private/FoliageEditModule.cpp
--- a/Engine/Source/Editor/FoliageEdit/Private/FoliageEditModule.cpp
+++ b/Engine/Source/Editor/FoliageEdit/Private/FoliageEditModule.cpp
@@ -14,23 +14,19 @@ const FName FoliageEditAppIdentifier = FName(TEXT("FoliageEdApp"));
  */
 class FFoliageEditModule : public IFoliageEditModule
 {
-private:
-	TSharedPtr<FEdModeFoliage> EdModeFoliage;
 public:
-	/** Constructor, set up console commands and variables **/
-	FFoliageEditModule()
-	{
-	}
 
 	/**
 	 * Called right after the module DLL has been loaded and the module object has been created
 	 */
 	virtual void StartupModule() override
 	{
-		TSharedRef<FEdModeFoliage> NewEditorMode = MakeShareable(new FEdModeFoliage);
-		GEditorModeTools().RegisterMode(NewEditorMode);
-		EdModeFoliage = NewEditorMode;
-		FEditorDelegates::MapChange.AddSP(NewEditorMode, &FEdModeFoliage::NotifyMapRebuild);
+		FEditorModeRegistry::Get().RegisterMode<FEdModeFoliage>(
+			FBuiltinEditorModes::EM_Foliage,
+			NSLOCTEXT("EditorModes", "FoliageMode", "Foliage"),
+			FSlateIcon(FEditorStyle::GetStyleSetName(), "LevelEditor.FoliageMode", "LevelEditor.FoliageMode.Small"),
+			true, 400
+			);
 	}
 
 	/**
@@ -38,10 +34,7 @@ public:
 	 */
 	virtual void ShutdownModule() override
 	{
-		FEditorDelegates::MapChange.RemoveAll(EdModeFoliage.Get());
-
-		GEditorModeTools().UnregisterMode(EdModeFoliage.ToSharedRef());
-		EdModeFoliage = NULL;
+		FEditorModeRegistry::Get().UnregisterMode(FBuiltinEditorModes::EM_Foliage);
 	}
 };
 
diff --git a/Engine/Source/Editor/FoliageEdit/Private/SFoliageEdit.cpp b/Engine/Source/Editor/FoliageEdit/Private/SFoliageEdit.cpp
--- a/Engine/Source/Editor/FoliageEdit/Private/SFoliageEdit.cpp
+++ b/Engine/Source/Editor/FoliageEdit/Private/SFoliageEdit.cpp
@@ -74,7 +74,7 @@ private:
 BEGIN_SLATE_FUNCTION_BUILD_OPTIMIZATION
 void SFoliageEdit::Construct(const FArguments& InArgs)
 {
-	FoliageEditMode = (FEdModeFoliage*)GEditorModeTools().GetActiveMode( FBuiltinEditorModes::EM_Foliage );
+	FoliageEditMode = (FEdModeFoliage*)GLevelEditorModeTools().GetActiveMode( FBuiltinEditorModes::EM_Foliage );
 
 	FFoliageEditCommands::Register();
 
diff --git a/Engine/Source/Editor/FoliageEdit/Private/SFoliageEditMeshDisplayItem.cpp b/Engine/Source/Editor/FoliageEdit/Private/SFoliageEditMeshDisplayItem.cpp
--- a/Engine/Source/Editor/FoliageEdit/Private/SFoliageEditMeshDisplayItem.cpp
+++ b/Engine/Source/Editor/FoliageEdit/Private/SFoliageEditMeshDisplayItem.cpp
@@ -1719,7 +1719,7 @@ FReply SFoliageEditMeshDisplayItem::OnReplace()
 		UStaticMesh* SelectedStaticMesh = Cast<UStaticMesh>( SelectedSet->GetTop( UStaticMesh::StaticClass() ) );
 		if( SelectedStaticMesh != NULL )
 		{
-			FEdModeFoliage* Mode = (FEdModeFoliage*)GEditorModeTools().GetActiveMode( FBuiltinEditorModes::EM_Foliage );
+			FEdModeFoliage* Mode = (FEdModeFoliage*)GLevelEditorModeTools().GetActiveMode( FBuiltinEditorModes::EM_Foliage );
 			
 			bool bMeshMerged = false;
 			if(Mode->ReplaceStaticMesh(AssociatedStaticMesh, SelectedStaticMesh, bMeshMerged))
@@ -1772,7 +1772,7 @@ FReply SFoliageEditMeshDisplayItem::OnSync()
 
 FReply SFoliageEditMeshDisplayItem::OnRemove()
 {
-	FEdModeFoliage* Mode = (FEdModeFoliage*)GEditorModeTools().GetActiveMode( FBuiltinEditorModes::EM_Foliage );
+	FEdModeFoliage* Mode = (FEdModeFoliage*)GLevelEditorModeTools().GetActiveMode( FBuiltinEditorModes::EM_Foliage );
 	
 	if(Mode->RemoveFoliageMesh(AssociatedStaticMesh))
 	{
@@ -1784,7 +1784,7 @@ FReply SFoliageEditMeshDisplayItem::OnRemove()
 
 FReply SFoliageEditMeshDisplayItem::OnSaveRemoveSettings()
 {
-	FEdModeFoliage* Mode = (FEdModeFoliage*)GEditorModeTools().GetActiveMode( FBuiltinEditorModes::EM_Foliage );
+	FEdModeFoliage* Mode = (FEdModeFoliage*)GLevelEditorModeTools().GetActiveMode( FBuiltinEditorModes::EM_Foliage );
 
 	if(FoliageSettingsPtr->GetOuter()->IsA(UPackage::StaticClass()))
 	{
@@ -1831,7 +1831,7 @@ FReply SFoliageEditMeshDisplayItem::OnOpenSettings()
 	UInstancedFoliageSettings* SelectedSettings = Cast<UInstancedFoliageSettings>( SelectedSet->GetTop( UInstancedFoliageSettings::StaticClass() ) );
 	if( SelectedSettings != NULL )
 	{
-		FEdModeFoliage* Mode = (FEdModeFoliage*)GEditorModeTools().GetActiveMode( FBuiltinEditorModes::EM_Foliage );
+		FEdModeFoliage* Mode = (FEdModeFoliage*)GLevelEditorModeTools().GetActiveMode( FBuiltinEditorModes::EM_Foliage );
 		Mode->ReplaceSettingsObject(AssociatedStaticMesh, SelectedSettings);
 		FoliageSettingsPtr = SelectedSettings;
 	}
@@ -1869,19 +1869,19 @@ FString SFoliageEditMeshDisplayItem::GetInstanceClusterCountString() const
 
 EVisibility SFoliageEditMeshDisplayItem::IsReapplySettingsVisible() const
 {
-	FEdModeFoliage* Mode = (FEdModeFoliage*)GEditorModeTools().GetActiveMode( FBuiltinEditorModes::EM_Foliage );
+	FEdModeFoliage* Mode = (FEdModeFoliage*)GLevelEditorModeTools().GetActiveMode( FBuiltinEditorModes::EM_Foliage );
 	return Mode->UISettings.GetReapplyToolSelected()? EVisibility::Visible : EVisibility::Collapsed;
 }
 
 EVisibility SFoliageEditMeshDisplayItem::IsReapplySettingsVisible_Dummy() const
 {
-	FEdModeFoliage* Mode = (FEdModeFoliage*)GEditorModeTools().GetActiveMode( FBuiltinEditorModes::EM_Foliage );
+	FEdModeFoliage* Mode = (FEdModeFoliage*)GLevelEditorModeTools().GetActiveMode( FBuiltinEditorModes::EM_Foliage );
 	return Mode->UISettings.GetReapplyToolSelected()? EVisibility::Hidden : EVisibility::Collapsed;
 }
 
 EVisibility SFoliageEditMeshDisplayItem::IsNotReapplySettingsVisible() const
 {
-	FEdModeFoliage* Mode = (FEdModeFoliage*)GEditorModeTools().GetActiveMode( FBuiltinEditorModes::EM_Foliage );
+	FEdModeFoliage* Mode = (FEdModeFoliage*)GLevelEditorModeTools().GetActiveMode( FBuiltinEditorModes::EM_Foliage );
 	return Mode->UISettings.GetReapplyToolSelected()? EVisibility::Collapsed : EVisibility::Visible;
 }
 
diff --git a/Engine/Source/Editor/GeometryMode/Classes/GeomModifier.h b/Engine/Source/Editor/GeometryMode/Classes/GeomModifier.h
--- a/Engine/Source/Editor/GeometryMode/Classes/GeomModifier.h
+++ b/Engine/Source/Editor/GeometryMode/Classes/GeomModifier.h
@@ -39,16 +39,16 @@ public:
 	const FText& GetModifierDescription() const;
 
 	/** @return		true if the key was handled by this editor mode tool. */
-	virtual bool InputKey(class FLevelEditorViewportClient* ViewportClient, FViewport* Viewport, FKey Key, EInputEvent Event);
+	virtual bool InputKey(class FEditorViewportClient* ViewportClient, FViewport* Viewport, FKey Key, EInputEvent Event);
 
 	/** @return		true if the delta was handled by this editor mode tool. */
-	virtual bool InputDelta(class FLevelEditorViewportClient* InViewportClient, FViewport* InViewport, FVector& InDrag, FRotator& InRot, FVector& InScale);
+	virtual bool InputDelta(class FEditorViewportClient* InViewportClient, FViewport* InViewport, FVector& InDrag, FRotator& InRot, FVector& InScale);
 
 	/*
 	 * Drawing functions to allow modifiers to have better control over the screen.
 	 */
 	virtual void Render(const FSceneView* View,FViewport* Viewport,FPrimitiveDrawInterface* PDI);
-	virtual void DrawHUD(FLevelEditorViewportClient* ViewportClient,FViewport* Viewport,const FSceneView* View,FCanvas* Canvas);
+	virtual void DrawHUD(FEditorViewportClient* ViewportClient,FViewport* Viewport,const FSceneView* View,FCanvas* Canvas);
 
 	/**
 	 * Applies the modifier.  Does nothing if the editor is not in geometry mode.
@@ -95,7 +95,7 @@ public:
 	 */
 	void StoreAllCurrentGeomSelections();
 
-	virtual void Tick(FLevelEditorViewportClient* ViewportClient,float DeltaTime) {}
+	virtual void Tick(FEditorViewportClient* ViewportClient,float DeltaTime) {}
 
 	/**
 	 * Gives the modifier a chance to initialize it's internal state when activated.
diff --git a/Engine/Source/Editor/GeometryMode/Classes/GeomModifier_Clip.h b/Engine/Source/Editor/GeometryMode/Classes/GeomModifier_Clip.h
--- a/Engine/Source/Editor/GeometryMode/Classes/GeomModifier_Clip.h
+++ b/Engine/Source/Editor/GeometryMode/Classes/GeomModifier_Clip.h
@@ -30,10 +30,10 @@ class UGeomModifier_Clip : public UGeomModifier_Edit
 
 	// Begin UGeomModifier Interface
 	virtual bool Supports() override;
-	virtual bool InputKey(class FLevelEditorViewportClient* ViewportClient, FViewport* Viewport, FKey Key, EInputEvent Event) override;
+	virtual bool InputKey(class FEditorViewportClient* ViewportClient, FViewport* Viewport, FKey Key, EInputEvent Event) override;
 	virtual void Render(const FSceneView* View,FViewport* Viewport,FPrimitiveDrawInterface* PDI) override;
-	virtual void DrawHUD(FLevelEditorViewportClient* ViewportClient,FViewport* Viewport,const FSceneView* View,FCanvas* Canvas) override;
-	virtual void Tick(FLevelEditorViewportClient* ViewportClient,float DeltaTime) override;
+	virtual void DrawHUD(FEditorViewportClient* ViewportClient,FViewport* Viewport,const FSceneView* View,FCanvas* Canvas) override;
+	virtual void Tick(FEditorViewportClient* ViewportClient,float DeltaTime) override;
 	virtual void WasActivated() override;
 protected:
 	virtual bool OnApply() override;
diff --git a/Engine/Source/Editor/GeometryMode/Classes/GeomModifier_Edit.h b/Engine/Source/Editor/GeometryMode/Classes/GeomModifier_Edit.h
--- a/Engine/Source/Editor/GeometryMode/Classes/GeomModifier_Edit.h
+++ b/Engine/Source/Editor/GeometryMode/Classes/GeomModifier_Edit.h
@@ -21,7 +21,7 @@ class UGeomModifier_Edit : public UGeomModifier
 
 
 	// Begin UGeomModifier Interface
-	virtual bool InputDelta(class FLevelEditorViewportClient* InViewportClient,FViewport* InViewport,FVector& InDrag,FRotator& InRot,FVector& InScale) override;	
+	virtual bool InputDelta(class FEditorViewportClient* InViewportClient,FViewport* InViewport,FVector& InDrag,FRotator& InRot,FVector& InScale) override;	
 	// End UGeomModifier Interface
 };
 
diff --git a/Engine/Source/Editor/GeometryMode/Classes/GeomModifier_Pen.h b/Engine/Source/Editor/GeometryMode/Classes/GeomModifier_Pen.h
--- a/Engine/Source/Editor/GeometryMode/Classes/GeomModifier_Pen.h
+++ b/Engine/Source/Editor/GeometryMode/Classes/GeomModifier_Pen.h
@@ -38,13 +38,13 @@ class UGeomModifier_Pen : public UGeomModifier_Edit
 	FVector MouseWorldSpacePos;
 
 
-	FLevelEditorViewportClient*		UsingViewportClient;
+	FEditorViewportClient*		UsingViewportClient;
 
 	// Begin UGeomModifier Interface
-	virtual bool InputKey(class FLevelEditorViewportClient* ViewportClient, FViewport* Viewport, FKey Key, EInputEvent Event) override;
+	virtual bool InputKey(class FEditorViewportClient* ViewportClient, FViewport* Viewport, FKey Key, EInputEvent Event) override;
 	virtual void Render(const FSceneView* View,FViewport* Viewport,FPrimitiveDrawInterface* PDI) override;
-	virtual void DrawHUD(FLevelEditorViewportClient* ViewportClient,FViewport* Viewport,const FSceneView* View,FCanvas* Canvas) override;
-	virtual void Tick(FLevelEditorViewportClient* ViewportClient,float DeltaTime) override;
+	virtual void DrawHUD(FEditorViewportClient* ViewportClient,FViewport* Viewport,const FSceneView* View,FCanvas* Canvas) override;
+	virtual void Tick(FEditorViewportClient* ViewportClient,float DeltaTime) override;
 	virtual void WasActivated() override;
 protected:
 	virtual bool OnApply() override;
diff --git a/Engine/Source/Editor/GeometryMode/Private/EditorGeometry.cpp b/Engine/Source/Editor/GeometryMode/Private/EditorGeometry.cpp
--- a/Engine/Source/Editor/GeometryMode/Private/EditorGeometry.cpp
+++ b/Engine/Source/Editor/GeometryMode/Private/EditorGeometry.cpp
@@ -21,10 +21,10 @@ FGeomBase::FGeomBase()
 
 void FGeomBase::Select( bool InSelect )
 {
-	FEditorModeTools& Tools = GEditorModeTools();
-	check( Tools.IsModeActive( FBuiltinEditorModes::EM_Geometry ) );
+	FEditorModeTools& Tools = GLevelEditorModeTools();
+	check(Tools.IsModeActive(FBuiltinEditorModes::EM_Geometry));
 
-	if( InSelect )
+	if (InSelect)
 	{
 		SelectionIndex = GetParentObject()->GetNewSelectionIndex();
 	}
@@ -44,19 +44,19 @@ void FGeomBase::Select( bool InSelect )
 
 FGeomObject* FGeomBase::GetParentObject()
 {
-	check( GEditorModeTools().IsModeActive(FBuiltinEditorModes::EM_Geometry) );
+	check( GLevelEditorModeTools().IsModeActive(FBuiltinEditorModes::EM_Geometry) );
 	check( ParentObjectIndex > INDEX_NONE );
 
-	FEdModeGeometry* mode = (FEdModeGeometry*)GEditorModeTools().GetActiveMode(FBuiltinEditorModes::EM_Geometry);
+	FEdModeGeometry* mode = (FEdModeGeometry*)GLevelEditorModeTools().GetActiveMode(FBuiltinEditorModes::EM_Geometry);
 	return mode->GetGeomObject( ParentObjectIndex );
 }
 
 const FGeomObject* FGeomBase::GetParentObject() const
 {
-	check( GEditorModeTools().IsModeActive(FBuiltinEditorModes::EM_Geometry) );
+	check( GLevelEditorModeTools().IsModeActive(FBuiltinEditorModes::EM_Geometry) );
 	check( ParentObjectIndex > INDEX_NONE );
 
-	const FEdModeGeometry* mode = (FEdModeGeometry*)GEditorModeTools().GetActiveMode(FBuiltinEditorModes::EM_Geometry);
+	const FEdModeGeometry* mode = (FEdModeGeometry*)GLevelEditorModeTools().GetActiveMode(FBuiltinEditorModes::EM_Geometry);
 	return mode->GetGeomObject( ParentObjectIndex );
 }
 
@@ -79,7 +79,7 @@ int32 FGeomObject::SetPivotFromSelectionArray( TArray<struct FGeomSelection>& Se
 	if( HighestSelectionIndex != -1 )
 	{
 		FGeomSelection* gs = &SelectionArray[ ArrayIndex ];
-		FEditorModeTools& Tools = GEditorModeTools();
+		FEditorModeTools& Tools = GLevelEditorModeTools();
 		switch( gs->Type )
 		{
 		case GS_Poly:
@@ -341,12 +341,12 @@ void FGeomObject::GetFromSource()
 
 int32 FGeomObject::GetObjectIndex()
 {
-	if( !GEditorModeTools().IsModeActive(FBuiltinEditorModes::EM_Geometry) )
+	if( !GLevelEditorModeTools().IsModeActive(FBuiltinEditorModes::EM_Geometry) )
 	{
 		return INDEX_NONE;
 	}
 
-	FEdModeGeometry* mode = (FEdModeGeometry*)GEditorModeTools().GetActiveMode(FBuiltinEditorModes::EM_Geometry);
+	FEdModeGeometry* mode = (FEdModeGeometry*)GLevelEditorModeTools().GetActiveMode(FBuiltinEditorModes::EM_Geometry);
 
 	for( FEdModeGeometry::TGeomObjectIterator Itor( mode->GeomObjectItor() ) ; Itor ; ++Itor )
 	{
@@ -389,7 +389,7 @@ void FGeomObject::SendToSource()
 
 bool FGeomObject::FinalizeSourceData()
 {
-	if( !GEditorModeTools().IsModeActive(FBuiltinEditorModes::EM_Geometry) )
+	if( !GLevelEditorModeTools().IsModeActive(FBuiltinEditorModes::EM_Geometry) )
 	{
 		return 0;
 	}
diff --git a/Engine/Source/Editor/GeometryMode/Private/GeometryEdMode.cpp b/Engine/Source/Editor/GeometryMode/Private/GeometryEdMode.cpp
--- a/Engine/Source/Editor/GeometryMode/Private/GeometryEdMode.cpp
+++ b/Engine/Source/Editor/GeometryMode/Private/GeometryEdMode.cpp
@@ -11,42 +11,25 @@ DEFINE_LOG_CATEGORY_STATIC(LogGeometryMode, Log, All);
 
 void FGeometryModeModule::StartupModule()
 {
-	EdModeGeometry = FEdModeGeometry::Create();
-	GEditorModeTools().RegisterMode( EdModeGeometry.ToSharedRef() );
-
-	IBspModeModule& BspModule = FModuleManager::LoadModuleChecked<IBspModeModule>( TEXT("BspMode") );
-	GEditorModeTools().RegisterCompatibleModes( EdModeGeometry->GetID(), BspModule.GetBspMode()->GetID() );
+	FEditorModeRegistry::Get().RegisterMode<FEdModeGeometry>(
+		FBuiltinEditorModes::EM_Geometry,
+		NSLOCTEXT("EditorModes", "GeometryMode", "Geometry Editing"),
+		FSlateIcon(FEditorStyle::GetStyleSetName(), "LevelEditor.BspMode", "LevelEditor.BspMode.Small"),
+		true
+		);
 }
 
 void FGeometryModeModule::ShutdownModule()
 {
-	if ( FModuleManager::Get().IsModuleLoaded( TEXT("BspMode") ) )
-	{
-		IBspModeModule& BspModule = FModuleManager::LoadModuleChecked<IBspModeModule>( TEXT("BspMode") );
-		GEditorModeTools().UnregisterCompatibleModes( EdModeGeometry->GetID(), BspModule.GetBspMode()->GetID() );
-	}
-
-	GEditorModeTools().UnregisterMode(EdModeGeometry.ToSharedRef());
-	EdModeGeometry = NULL;
+	FEditorModeRegistry::Get().UnregisterMode(FBuiltinEditorModes::EM_Geometry);
 }
 
 /*------------------------------------------------------------------------------
 	Geometry Editing.
 ------------------------------------------------------------------------------*/
 
-TSharedRef< FEdModeGeometry > FEdModeGeometry::Create()
-{
-	TSharedRef< FEdModeGeometry > GeometryMode = MakeShareable( new FEdModeGeometry() );
-	return GeometryMode;
-}
-
-
 FEdModeGeometry::FEdModeGeometry()
 {
-	ID = FBuiltinEditorModes::EM_Geometry;
-	Name = NSLOCTEXT("EditorModes", "GeometryMode", "Geometry Editing");
-	bVisible = true;
-	IconBrush = FSlateIcon(FEditorStyle::GetStyleSetName(), "LevelEditor.BspMode", "LevelEditor.BspMode.Small");
 	Tools.Add( new FModeTool_GeometryModify() );
 	SetCurrentTool( MT_GeometryModify );
 }
@@ -430,6 +413,11 @@ FVector FEdModeGeometry::GetWidgetLocation() const
 	return FEdMode::GetWidgetLocation();
 }
 
+bool FEdModeGeometry::IsCompatibleWith(FEditorModeID OtherModeID) const
+{
+	return OtherModeID == FBuiltinEditorModes::EM_Bsp;
+}
+
 // ------------------------------------------------------------------------------
 
 /**
@@ -466,8 +454,7 @@ void FEdModeGeometry::GeometrySelectNone(bool bStoreSelection, bool bResetPivot)
 
 	if (bResetPivot && (GeomObjects.Num() > 0))
 	{
-		FEditorModeTools& ModeTools = GEditorModeTools();
-		ModeTools.SetPivotLocation(GeomObjects[0]->GetActualBrush()->GetActorLocation(), false);
+		Owner->SetPivotLocation(GeomObjects[0]->GetActualBrush()->GetActorLocation(), false);
 	}
 }
 
@@ -815,13 +802,11 @@ bool FEdModeGeometry::SelectCachedData( TArray<FGeomBase*>& raSelectedGeom ) con
 		return false;
 	}
 
-	// Grab the editor tools so we can reposition the widget correctly
-	FEditorModeTools& Tools = GEditorModeTools();
-	check( Tools.IsModeActive(FBuiltinEditorModes::EM_Geometry) );
+	check(Owner->IsModeActive(FBuiltinEditorModes::EM_Geometry));
 
 	// Backup widget position, we want it to be in the same position as it was previously too
-	FVector PivLoc = Tools.PivotLocation;
-	FVector SnapLoc = Tools.SnappedLocation;
+	FVector PivLoc = Owner->PivotLocation;
+	FVector SnapLoc = Owner->SnappedLocation;
 
 	// Loop through all the geometry that should be selected
 	for( int32 i=0; i<raSelectedGeom.Num(); i++ )
@@ -835,8 +820,8 @@ bool FEdModeGeometry::SelectCachedData( TArray<FGeomBase*>& raSelectedGeom ) con
 	}
 
 	// Restore the widget position
-	Tools.SetPivotLocation( PivLoc, false );
-	Tools.SnappedLocation = SnapLoc;
+	Owner->SetPivotLocation( PivLoc, false );
+	Owner->SnappedLocation = SnapLoc;
 	
 	StaticCastSharedPtr<FGeometryMode>(Toolkit)->SelectionChanged();
 
@@ -973,8 +958,7 @@ void FEdModeGeometry::PostUndo()
 		//use the centre of the actor if we didnt find a suitable selection
 		if( res == INDEX_NONE )
 		{
-			FEditorModeTools& Tools = GEditorModeTools();
-			Tools.SetPivotLocation( Actor->GetActorLocation() , false );
+			Owner->SetPivotLocation( Actor->GetActorLocation() , false );
 		}
 		
 		go->ForceLastSelectionIndex( res );
@@ -983,9 +967,7 @@ void FEdModeGeometry::PostUndo()
 
 bool FEdModeGeometry::ExecDelete()
 {
-	check( GEditorModeTools().IsModeActive( FBuiltinEditorModes::EM_Geometry ) );
-
-	FEdModeGeometry* mode = (FEdModeGeometry*)GEditorModeTools().GetActiveMode( FBuiltinEditorModes::EM_Geometry );
+	check( Owner->IsModeActive( FBuiltinEditorModes::EM_Geometry ) );
 
 	// Find the delete modifier and execute it.
 
@@ -1058,7 +1040,7 @@ FModeTool_GeometryModify::FModeTool_GeometryModify()
 
 void FModeTool_GeometryModify::SelectNone()
 {
-	FEdModeGeometry* Mode = GEditorModeTools().GetActiveModeTyped<FEdModeGeometry>( FBuiltinEditorModes::EM_Geometry );
+	FEdModeGeometry* Mode = GLevelEditorModeTools().GetActiveModeTyped<FEdModeGeometry>(FBuiltinEditorModes::EM_Geometry);
 	check(Mode);
 	Mode->GeometrySelectNone(false, false);
 }
@@ -1067,9 +1049,9 @@ void FModeTool_GeometryModify::SelectNone()
 bool FModeTool_GeometryModify::BoxSelect( FBox& InBox, bool InSelect )
 {
 	bool bResult = false;
-	if( GEditorModeTools().IsModeActive( FBuiltinEditorModes::EM_Geometry ) )
+	if( GLevelEditorModeTools().IsModeActive( FBuiltinEditorModes::EM_Geometry ) )
 	{
-		FEdModeGeometry* mode = (FEdModeGeometry*)GEditorModeTools().GetActiveMode( FBuiltinEditorModes::EM_Geometry );
+		FEdModeGeometry* mode = (FEdModeGeometry*)GLevelEditorModeTools().GetActiveMode( FBuiltinEditorModes::EM_Geometry );
 
 		for( FEdModeGeometry::TGeomObjectIterator Itor( mode->GeomObjectItor() ) ; Itor ; ++Itor )
 		{
@@ -1096,9 +1078,9 @@ bool FModeTool_GeometryModify::BoxSelect( FBox& InBox, bool InSelect )
 bool FModeTool_GeometryModify::FrustumSelect( const FConvexVolume& InFrustum, bool InSelect /* = true */ )
 {
 	bool bResult = false;
-	if( GEditorModeTools().IsModeActive( FBuiltinEditorModes::EM_Geometry ) )
+	if( GLevelEditorModeTools().IsModeActive( FBuiltinEditorModes::EM_Geometry ) )
 	{
-		FEdModeGeometry* mode = (FEdModeGeometry*)GEditorModeTools().GetActiveMode( FBuiltinEditorModes::EM_Geometry );
+		FEdModeGeometry* mode = (FEdModeGeometry*)GLevelEditorModeTools().GetActiveMode( FBuiltinEditorModes::EM_Geometry );
 
 		for( FEdModeGeometry::TGeomObjectIterator Itor( mode->GeomObjectItor() ) ; Itor ; ++Itor )
 		{
@@ -1119,7 +1101,7 @@ bool FModeTool_GeometryModify::FrustumSelect( const FConvexVolume& InFrustum, bo
 	return bResult;
 }
 
-void FModeTool_GeometryModify::Tick(FLevelEditorViewportClient* ViewportClient,float DeltaTime)
+void FModeTool_GeometryModify::Tick(FEditorViewportClient* ViewportClient,float DeltaTime)
 {
 	if( CurrentModifier )
 	{
@@ -1130,7 +1112,7 @@ void FModeTool_GeometryModify::Tick(FLevelEditorViewportClient* ViewportClient,f
 /**
  * @return		true if the delta was handled by this editor mode tool.
  */
-bool FModeTool_GeometryModify::InputDelta(FLevelEditorViewportClient* InViewportClient,FViewport* InViewport,FVector& InDrag,FRotator& InRot,FVector& InScale)
+bool FModeTool_GeometryModify::InputDelta(FEditorViewportClient* InViewportClient,FViewport* InViewport,FVector& InDrag,FRotator& InRot,FVector& InScale)
 {
 	bool bResult = false;
 	if( InViewportClient->GetCurrentWidgetAxis() != EAxisList::None )
@@ -1167,7 +1149,7 @@ bool FModeTool_GeometryModify::EndModify()
 	// Let the modifier finish up.
 	if( CurrentModifier != NULL )
 	{
-		FEdModeGeometry* mode = ((FEdModeGeometry*)GEditorModeTools().GetActiveMode(FBuiltinEditorModes::EM_Geometry));
+		FEdModeGeometry* mode = ((FEdModeGeometry*)GLevelEditorModeTools().GetActiveMode(FBuiltinEditorModes::EM_Geometry));
 
 		// Update the source data to match the current geometry data.
 		mode->SendToSource();
@@ -1226,9 +1208,9 @@ void FModeTool_GeometryModify::EndTrans()
 /**
  * @return		true if the key was handled by this editor mode tool.
  */
-bool FModeTool_GeometryModify::InputKey(FLevelEditorViewportClient* ViewportClient, FViewport* Viewport, FKey Key, EInputEvent Event)
+bool FModeTool_GeometryModify::InputKey(FEditorViewportClient* ViewportClient, FViewport* Viewport, FKey Key, EInputEvent Event)
 {
-	check( GEditorModeTools().IsModeActive( FBuiltinEditorModes::EM_Geometry ) );
+	check( GLevelEditorModeTools().IsModeActive( FBuiltinEditorModes::EM_Geometry ) );
 	// Give the current modifier a chance to handle this first
 	if( CurrentModifier && CurrentModifier->InputKey( ViewportClient, Viewport, Key, Event ) )
 	{
@@ -1240,7 +1222,7 @@ bool FModeTool_GeometryModify::InputKey(FLevelEditorViewportClient* ViewportClie
 		// Hitting ESC will deselect any subobjects first.  If no subobjects are selected, then it will
 		// deselect the brushes themselves.
 
-		FEdModeGeometry* mode = (FEdModeGeometry*)GEditorModeTools().GetActiveMode( FBuiltinEditorModes::EM_Geometry );
+		FEdModeGeometry* mode = (FEdModeGeometry*)GLevelEditorModeTools().GetActiveMode( FBuiltinEditorModes::EM_Geometry );
 		bool bHadSubObjectSelections = (mode->GetSelectionState() > 0) ? true : false;
 
 		for( FEdModeGeometry::TGeomObjectIterator Itor( mode->GeomObjectItor() ) ; Itor ; ++Itor )
@@ -1292,7 +1274,7 @@ bool FModeTool_GeometryModify::InputKey(FLevelEditorViewportClient* ViewportClie
 	return false;
 }
 
-void FModeTool_GeometryModify::DrawHUD(FLevelEditorViewportClient* ViewportClient,FViewport* Viewport,const FSceneView* View,FCanvas* Canvas)
+void FModeTool_GeometryModify::DrawHUD(FEditorViewportClient* ViewportClient,FViewport* Viewport,const FSceneView* View,FCanvas* Canvas)
 {
 	// Give the current modifier a chance to draw a HUD
 
diff --git a/Engine/Source/Editor/GeometryMode/Private/GeometryMode.cpp b/Engine/Source/Editor/GeometryMode/Private/GeometryMode.cpp
--- a/Engine/Source/Editor/GeometryMode/Private/GeometryMode.cpp
+++ b/Engine/Source/Editor/GeometryMode/Private/GeometryMode.cpp
@@ -71,7 +71,7 @@ EVisibility SGeometryModeControls::IsPropertiesVisible() const
 
 FReply SGeometryModeControls::OnApplyClicked()
 {
-	check(GEditorModeTools().IsModeActive(FBuiltinEditorModes::EM_Geometry));
+	check(GLevelEditorModeTools().IsModeActive(FBuiltinEditorModes::EM_Geometry));
 
 	GetGeometryModeTool()->GetCurrentModifier()->Apply();
 
@@ -80,7 +80,7 @@ FReply SGeometryModeControls::OnApplyClicked()
 
 FReply SGeometryModeControls::OnModifierClicked(UGeomModifier* Modifier)
 {
-	check(GEditorModeTools().IsModeActive(FBuiltinEditorModes::EM_Geometry));
+	check(GLevelEditorModeTools().IsModeActive(FBuiltinEditorModes::EM_Geometry));
 
 	Modifier->Apply();
 
@@ -261,7 +261,7 @@ TSharedRef<SButton> SGeometryModeControls::CreateSingleModifierButton(UGeomModif
 
 FModeTool_GeometryModify* SGeometryModeControls::GetGeometryModeTool() const
 {
-	FEdModeGeometry* Mode = (FEdModeGeometry*)GEditorModeTools().GetActiveMode( FBuiltinEditorModes::EM_Geometry );
+	FEdModeGeometry* Mode = (FEdModeGeometry*)GLevelEditorModeTools().GetActiveMode( FBuiltinEditorModes::EM_Geometry );
 	FModeTool* Tool = Mode? Mode->GetCurrentTool(): NULL;
 
 	check(Tool);
@@ -298,7 +298,7 @@ FText FGeometryMode::GetBaseToolkitName() const
 
 class FEdMode* FGeometryMode::GetEditorMode() const
 {
-	return (FEdModeGeometry*)GEditorModeTools().GetActiveMode(FBuiltinEditorModes::EM_Geometry);
+	return (FEdModeGeometry*)GLevelEditorModeTools().GetActiveMode(FBuiltinEditorModes::EM_Geometry);
 }
 
 void FGeometryMode::SelectionChanged()
