Made it possible for asset editors to maintain their own FEditorModeTools lists

diff --git a/Engine/Source/Editor/GeometryMode/Private/GeometryModifiers.cpp b/Engine/Source/Editor/GeometryMode/Private/GeometryModifiers.cpp
--- a/Engine/Source/Editor/GeometryMode/Private/GeometryModifiers.cpp
+++ b/Engine/Source/Editor/GeometryMode/Private/GeometryModifiers.cpp
@@ -13,7 +13,7 @@ DEFINE_LOG_CATEGORY_STATIC(LogGeomModifier, Log, All);
 
 #define LOCTEXT_NAMESPACE "UnrealEd.GeomModifier"
 
-static FVector ComputeWorldSpaceMousePos( FLevelEditorViewportClient* ViewportClient )
+static FVector ComputeWorldSpaceMousePos( FEditorViewportClient* ViewportClient )
 {
 	FSceneViewFamilyContext ViewFamily(FSceneViewFamily::ConstructionValues(
 		ViewportClient->Viewport,
@@ -44,14 +44,14 @@ void UGeomModifier::Initialize()
 {
 }
 
-bool UGeomModifier::InputKey(FLevelEditorViewportClient* ViewportClient, FViewport* Viewport, FKey Key, EInputEvent Event)
+bool UGeomModifier::InputKey(FEditorViewportClient* ViewportClient, FViewport* Viewport, FKey Key, EInputEvent Event)
 {
 	return false;
 }
 
-bool UGeomModifier::InputDelta(FLevelEditorViewportClient* InViewportClient,FViewport* InViewport,FVector& InDrag,FRotator& InRot,FVector& InScale)
+bool UGeomModifier::InputDelta(FEditorViewportClient* InViewportClient,FViewport* InViewport,FVector& InDrag,FRotator& InRot,FVector& InScale)
 {
-	if( GEditorModeTools().IsModeActive(FBuiltinEditorModes::EM_Geometry) )
+	if( GLevelEditorModeTools().IsModeActive(FBuiltinEditorModes::EM_Geometry) )
 	{
 		if( !bInitialized )
 		{
@@ -67,7 +67,7 @@ bool UGeomModifier::InputDelta(FLevelEditorViewportClient* InViewportClient,FVie
 bool UGeomModifier::Apply()
 {
 	bool bResult = false;
-	if( GEditorModeTools().IsModeActive(FBuiltinEditorModes::EM_Geometry) )
+	if( GLevelEditorModeTools().IsModeActive(FBuiltinEditorModes::EM_Geometry) )
 	{
 		StartTrans();
 		bResult = OnApply();
@@ -114,14 +114,14 @@ void UGeomModifier::Render(const FSceneView* View,FViewport* Viewport,FPrimitive
 {
 }
 
-void UGeomModifier::DrawHUD(FLevelEditorViewportClient* ViewportClient,FViewport* Viewport,const FSceneView* View,FCanvas* Canvas)
+void UGeomModifier::DrawHUD(FEditorViewportClient* ViewportClient,FViewport* Viewport,const FSceneView* View,FCanvas* Canvas)
 {
 }
 
 
 void UGeomModifier::CacheBrushState()
 {
-	FEdModeGeometry* GeomMode = (FEdModeGeometry*)GEditorModeTools().GetActiveMode(FBuiltinEditorModes::EM_Geometry);
+	FEdModeGeometry* GeomMode = (FEdModeGeometry*)GLevelEditorModeTools().GetActiveMode(FBuiltinEditorModes::EM_Geometry);
 	ABrush* BuilderBrush = GeomMode->GetWorld()->GetDefaultBrush();
 	if( !CachedPolys )
 	{
@@ -151,7 +151,7 @@ void UGeomModifier::CacheBrushState()
 
 void UGeomModifier::RestoreBrushState()
 {
-	FEdModeGeometry* GeomMode = (FEdModeGeometry*)GEditorModeTools().GetActiveMode(FBuiltinEditorModes::EM_Geometry);
+	FEdModeGeometry* GeomMode = (FEdModeGeometry*)GLevelEditorModeTools().GetActiveMode(FBuiltinEditorModes::EM_Geometry);
 	ABrush* BuilderBrush = GeomMode->GetWorld()->GetDefaultBrush();
 
 	//Remove all of the current polys
@@ -181,7 +181,7 @@ void UGeomModifier::RestoreBrushState()
 
 bool UGeomModifier::DoEdgesOverlap()
 {
-	FEdModeGeometry* mode = (FEdModeGeometry*)GEditorModeTools().GetActiveMode(FBuiltinEditorModes::EM_Geometry);
+	FEdModeGeometry* mode = (FEdModeGeometry*)GLevelEditorModeTools().GetActiveMode(FBuiltinEditorModes::EM_Geometry);
 
 	//Loop through all of the geometry objects
 	for( FEdModeGeometry::TGeomObjectIterator itor( mode->GeomObjectItor() ) ; itor ; ++itor )
@@ -276,7 +276,7 @@ namespace {
 
 void UGeomModifier::StartTrans()
 {
-	if( !GEditorModeTools().IsModeActive(FBuiltinEditorModes::EM_Geometry) )
+	if( !GLevelEditorModeTools().IsModeActive(FBuiltinEditorModes::EM_Geometry) )
 	{
 		return;
 	}
@@ -287,7 +287,7 @@ void UGeomModifier::StartTrans()
 	BeginTransaction( FText::Format( NSLOCTEXT("UnrealEd", "Modifier_F", "Modifier [{0}]"), GetModifierDescription() ) );
 
 	// Mark all selected brushes as modified.
-	FEdModeGeometry* CurMode = static_cast<FEdModeGeometry*>( GEditorModeTools().GetActiveMode(FBuiltinEditorModes::EM_Geometry) );
+	FEdModeGeometry* CurMode = static_cast<FEdModeGeometry*>( GLevelEditorModeTools().GetActiveMode(FBuiltinEditorModes::EM_Geometry) );
 	for( FEdModeGeometry::TGeomObjectIterator Itor( CurMode->GeomObjectItor() ) ; Itor ; ++Itor )
 	{
 		FGeomObject* go = *Itor;
@@ -346,12 +346,12 @@ void UGeomModifier::StoreCurrentGeomSelections( TArray<struct FGeomSelection>& S
 
 void UGeomModifier::StoreAllCurrentGeomSelections()
 {
-	if( !GEditorModeTools().IsModeActive(FBuiltinEditorModes::EM_Geometry) )
+	if( !GLevelEditorModeTools().IsModeActive(FBuiltinEditorModes::EM_Geometry) )
 	{
 		return;
 	}
 
-	FEdModeGeometry* CurMode = static_cast<FEdModeGeometry*>( GEditorModeTools().GetActiveMode(FBuiltinEditorModes::EM_Geometry) );
+	FEdModeGeometry* CurMode = static_cast<FEdModeGeometry*>( GLevelEditorModeTools().GetActiveMode(FBuiltinEditorModes::EM_Geometry) );
 
 	// Record the current selection list into the selected brushes.
 	for( FEdModeGeometry::TGeomObjectIterator Itor( CurMode->GeomObjectItor() ) ; Itor ; ++Itor )
@@ -373,19 +373,19 @@ UGeomModifier_Edit::UGeomModifier_Edit(const class FPostConstructInitializePrope
 }
 
 
-bool UGeomModifier_Edit::InputDelta(FLevelEditorViewportClient* InViewportClient,FViewport* InViewport,FVector& InDrag,FRotator& InRot,FVector& InScale)
+bool UGeomModifier_Edit::InputDelta(FEditorViewportClient* InViewportClient,FViewport* InViewport,FVector& InDrag,FRotator& InRot,FVector& InScale)
 {
 	if( UGeomModifier::InputDelta( InViewportClient, InViewport, InDrag, InRot, InScale ) )
 	{
 		return true;
 	}
 
-	if( !GEditorModeTools().IsModeActive(FBuiltinEditorModes::EM_Geometry) )
+	if( !GLevelEditorModeTools().IsModeActive(FBuiltinEditorModes::EM_Geometry) )
 	{
 		return false;
 	}
 
-	FEdModeGeometry* mode = (FEdModeGeometry*)GEditorModeTools().GetActiveMode(FBuiltinEditorModes::EM_Geometry);
+	FEdModeGeometry* mode = (FEdModeGeometry*)GLevelEditorModeTools().GetActiveMode(FBuiltinEditorModes::EM_Geometry);
 	FModeTool_GeometryModify* tool = (FModeTool_GeometryModify*)mode->GetCurrentTool();
 
 	TArray<FGeomVertex*> UniqueVertexList;
@@ -516,9 +516,9 @@ bool UGeomModifier_Edit::InputDelta(FLevelEditorViewportClient* InViewportClient
 
 			FVector Wk( vtx->X, vtx->Y, vtx->Z );
 			Wk = vtx->GetParentObject()->GetActualBrush()->ActorToWorld().TransformPosition( Wk );
-			Wk -= GEditorModeTools().PivotLocation;
+			Wk -= GLevelEditorModeTools().PivotLocation;
 			Wk = matrix.TransformPosition( Wk );
-			Wk += GEditorModeTools().PivotLocation;
+			Wk += GLevelEditorModeTools().PivotLocation;
 			*vtx = vtx->GetParentObject()->GetActualBrush()->ActorToWorld().InverseTransformPosition( Wk );
 		}
 
@@ -536,7 +536,7 @@ bool UGeomModifier_Edit::InputDelta(FLevelEditorViewportClient* InViewportClient
 
 			// Move vert to the origin
 
-			Wk -= GEditorModeTools().PivotLocation;
+			Wk -= GLevelEditorModeTools().PivotLocation;
 
 			// Move it along each axis based on it's distance from the origin
 
@@ -560,7 +560,7 @@ bool UGeomModifier_Edit::InputDelta(FLevelEditorViewportClient* InViewportClient
 
 			// Move it back into world space
 
-			Wk += GEditorModeTools().PivotLocation;
+			Wk += GLevelEditorModeTools().PivotLocation;
 
 			*vtx = vtx->GetParentObject()->GetActualBrush()->ActorToWorld().InverseTransformPosition( Wk );
 		}
@@ -575,8 +575,8 @@ bool UGeomModifier_Edit::InputDelta(FLevelEditorViewportClient* InViewportClient
 			*vtx -= vertOffset;
 		}
 		
-		GEditorModeTools().PivotLocation -= vertOffset;
-		GEditorModeTools().SnappedLocation -= vertOffset;
+		GLevelEditorModeTools().PivotLocation -= vertOffset;
+		GLevelEditorModeTools().SnappedLocation -= vertOffset;
 	}
 
 	const bool bIsCtrlPressed = InViewportClient->IsCtrlPressed();
@@ -616,7 +616,7 @@ UGeomModifier_Extrude::UGeomModifier_Extrude(const class FPostConstructInitializ
 
 bool UGeomModifier_Extrude::Supports()
 {
-	FEdModeGeometry* mode = (FEdModeGeometry*)GEditorModeTools().GetActiveMode(FBuiltinEditorModes::EM_Geometry);
+	FEdModeGeometry* mode = (FEdModeGeometry*)GLevelEditorModeTools().GetActiveMode(FBuiltinEditorModes::EM_Geometry);
 	return mode->HavePolygonsSelected();
 }
 
@@ -625,7 +625,7 @@ void UGeomModifier_Extrude::WasActivated()
 	// Extrude requires a local coordinate system to work properly so automatically enable
 	// that here while saving the current coordinate system for restoration later.
 	const bool bGetRawValue = true;
-	SaveCoordSystem = GEditorModeTools().GetCoordSystem(bGetRawValue);
+	SaveCoordSystem = GLevelEditorModeTools().GetCoordSystem(bGetRawValue);
 	
 	CheckCoordinatesMode();
 
@@ -635,7 +635,7 @@ void UGeomModifier_Extrude::WasActivated()
 void UGeomModifier_Extrude::WasDeactivated()
 {
 	// When the user leaves this modifier, restore their old coordinate system.
-	GEditorModeTools().SetCoordSystem((ECoordSystem)SaveCoordSystem);
+	GLevelEditorModeTools().SetCoordSystem((ECoordSystem)SaveCoordSystem);
 
 	GEditor->RedrawLevelEditingViewports(true);
 }
@@ -643,14 +643,14 @@ void UGeomModifier_Extrude::WasDeactivated()
 void UGeomModifier_Extrude::CheckCoordinatesMode() 
 {
 	const bool bGetRawValue = true;	
-	if( GEditorModeTools().GetCoordSystem(bGetRawValue) != COORD_Local )
+	if( GLevelEditorModeTools().GetCoordSystem(bGetRawValue) != COORD_Local )
 	{
 		FSuppressableWarningDialog::FSetupInfo Info( LOCTEXT("ExtrudeCoordinateWarningBody","Extrude only works with Local Coordinates System"), LOCTEXT("ExtrudeCoordinateWarningTitle","Extrude Coordinates Mode Warning"), "ExtrudeCoordsWarning" );
 		Info.ConfirmText = LOCTEXT( "Close", "Close");
 
 		FSuppressableWarningDialog WarnAboutCoordinatesSystem( Info );
 		WarnAboutCoordinatesSystem.ShowModal();			
-		GEditorModeTools().SetCoordSystem(COORD_Local);
+		GLevelEditorModeTools().SetCoordSystem(COORD_Local);
 	}
 }
 
@@ -665,18 +665,18 @@ void UGeomModifier_Extrude::Initialize()
 
 bool UGeomModifier_Extrude::OnApply()
 {
-	FEdModeGeometry* mode = (FEdModeGeometry*)GEditorModeTools().GetActiveMode(FBuiltinEditorModes::EM_Geometry);
+	FEdModeGeometry* mode = (FEdModeGeometry*)GLevelEditorModeTools().GetActiveMode(FBuiltinEditorModes::EM_Geometry);
 
 	// When applying via the keyboard, we force the local coordinate system.
 	const bool bGetRawValue = true;
-	const ECoordSystem SaveCS = GEditorModeTools().GetCoordSystem(bGetRawValue);
-	GEditorModeTools().SetCoordSystem(COORD_Local);
+	const ECoordSystem SaveCS = GLevelEditorModeTools().GetCoordSystem(bGetRawValue);
+	GLevelEditorModeTools().SetCoordSystem(COORD_Local);
 
 	//GApp->DlgGeometryTools->PropertyWindow->FinalizeValues();
 	Apply( Length, Segments );
 
 	// Restore the coordinate system.
-	GEditorModeTools().SetCoordSystem(SaveCS);
+	GLevelEditorModeTools().SetCoordSystem(SaveCS);
 
 	GEditor->RebuildAlteredBSP(); // Brush has been altered, update the Bsp
 
@@ -722,12 +722,12 @@ void ExtrudePolygonGroup( ABrush* InBrush, FVector InGroupNormal, int32 InStartO
 
 void UGeomModifier_Extrude::Apply(int32 InLength, int32 InSegments)
 {
-	if( !GEditorModeTools().IsModeActive(FBuiltinEditorModes::EM_Geometry) )
+	if( !GLevelEditorModeTools().IsModeActive(FBuiltinEditorModes::EM_Geometry) )
 	{
 		return;
 	}
 
-	FEdModeGeometry* mode = (FEdModeGeometry*)GEditorModeTools().GetActiveMode(FBuiltinEditorModes::EM_Geometry);
+	FEdModeGeometry* mode = (FEdModeGeometry*)GLevelEditorModeTools().GetActiveMode(FBuiltinEditorModes::EM_Geometry);
 
 	// Force user input to be valid
 
@@ -872,7 +872,7 @@ void UGeomModifier_Lathe::Initialize()
 
 bool UGeomModifier_Lathe::OnApply()
 {
-	//FEdModeGeometry* mode = (FEdModeGeometry*)GEditorModeTools().GetCurrentMode();
+	//FEdModeGeometry* mode = (FEdModeGeometry*)GLevelEditorModeTools().GetCurrentMode();
 
 	//GApp->DlgGeometryTools->PropertyWindow->FinalizeValues();
 	Apply( TotalSegments, Segments, Axis );
@@ -884,7 +884,7 @@ bool UGeomModifier_Lathe::OnApply()
 
 void UGeomModifier_Lathe::Apply( int32 InTotalSegments, int32 InSegments, EAxis::Type InAxis )
 {
-	if( !GEditorModeTools().IsModeActive(FBuiltinEditorModes::EM_Geometry) )
+	if( !GLevelEditorModeTools().IsModeActive(FBuiltinEditorModes::EM_Geometry) )
 	{
 		return;
 	}
@@ -914,7 +914,7 @@ void UGeomModifier_Lathe::Apply( int32 InTotalSegments, int32 InSegments, EAxis:
 			break;
 	}
 
-	FEdModeGeometry* GeomMode = (FEdModeGeometry*)GEditorModeTools().GetActiveMode(FBuiltinEditorModes::EM_Geometry);
+	FEdModeGeometry* GeomMode = (FEdModeGeometry*)GLevelEditorModeTools().GetActiveMode(FBuiltinEditorModes::EM_Geometry);
 
 	InTotalSegments = FMath::Max( 3, InTotalSegments );
 	InSegments = FMath::Max( 1, InSegments );
@@ -1199,7 +1199,7 @@ void UGeomModifier_Pen::Apply()
 {
 	if( ShapeVertices.Num() > 2 )
 	{
-		FEdModeGeometry* GeomMode = (FEdModeGeometry*)GEditorModeTools().GetActiveMode(FBuiltinEditorModes::EM_Geometry);
+		FEdModeGeometry* GeomMode = (FEdModeGeometry*)GLevelEditorModeTools().GetActiveMode(FBuiltinEditorModes::EM_Geometry);
 		ABrush* ResultingBrush = GeomMode->GetWorld()->GetDefaultBrush();
 		ABrush* BuilderBrush = GeomMode->GetWorld()->GetDefaultBrush();
 
@@ -1380,7 +1380,7 @@ void UGeomModifier_Pen::Apply()
 
 		ShapeVertices.Empty();
 
-		FEdModeGeometry* mode = (FEdModeGeometry*)GEditorModeTools().GetActiveMode(FBuiltinEditorModes::EM_Geometry);
+		FEdModeGeometry* mode = (FEdModeGeometry*)GLevelEditorModeTools().GetActiveMode(FBuiltinEditorModes::EM_Geometry);
 
 		mode->FinalizeSourceData();
 		mode->GetFromSource();
@@ -1402,7 +1402,7 @@ void UGeomModifier_Pen::Apply()
 /**
 * @return		true if the key was handled by this editor mode tool.
 */
-bool UGeomModifier_Pen::InputKey(FLevelEditorViewportClient* ViewportClient, FViewport* Viewport, FKey Key, EInputEvent Event)
+bool UGeomModifier_Pen::InputKey(FEditorViewportClient* ViewportClient, FViewport* Viewport, FKey Key, EInputEvent Event)
 {
 	bool bResult = false;
 #if WITH_EDITORONLY_DATA
@@ -1462,7 +1462,7 @@ bool UGeomModifier_Pen::InputKey(FLevelEditorViewportClient* ViewportClient, FVi
 
 void UGeomModifier_Pen::Render(const FSceneView* View,FViewport* Viewport,FPrimitiveDrawInterface* PDI)
 {
-	FEdModeGeometry* mode = (FEdModeGeometry*)GEditorModeTools().GetActiveMode(FBuiltinEditorModes::EM_Geometry);
+	FEdModeGeometry* mode = (FEdModeGeometry*)GLevelEditorModeTools().GetActiveMode(FBuiltinEditorModes::EM_Geometry);
 	FModeTool_GeometryModify* tool = (FModeTool_GeometryModify*)mode->GetCurrentTool();
 	if( tool->GetCurrentModifier() != this )
 	{
@@ -1471,7 +1471,7 @@ void UGeomModifier_Pen::Render(const FSceneView* View,FViewport* Viewport,FPrimi
 
 	// Only draw in ortho viewports
 
-	if( !((FLevelEditorViewportClient*)(Viewport->GetClient()))->IsOrtho() )
+	if( !((FEditorViewportClient*)(Viewport->GetClient()))->IsOrtho() )
 	{
 		return;
 	}
@@ -1513,11 +1513,11 @@ void UGeomModifier_Pen::Render(const FSceneView* View,FViewport* Viewport,FPrimi
 	DrawWireBox(PDI, FBox::BuildAABB( MouseWorldSpacePos, FVector(BoxSz,BoxSz,BoxSz) ), FLinearColor(1,1,1), SDPG_Foreground);
 }
 
-void UGeomModifier_Pen::DrawHUD(FLevelEditorViewportClient* ViewportClient,FViewport* Viewport,const FSceneView* View,FCanvas* Canvas)
+void UGeomModifier_Pen::DrawHUD(FEditorViewportClient* ViewportClient,FViewport* Viewport,const FSceneView* View,FCanvas* Canvas)
 {
 }
 
-void UGeomModifier_Pen::Tick(FLevelEditorViewportClient* ViewportClient,float DeltaTime)
+void UGeomModifier_Pen::Tick(FEditorViewportClient* ViewportClient,float DeltaTime)
 {
 	if( GCurrentLevelEditingViewportClient == ViewportClient )
 	{
@@ -1831,7 +1831,7 @@ void UGeomModifier_Clip::WasActivated()
 
 bool UGeomModifier_Clip::Supports()
 {
-	FEdModeGeometry* mode = (FEdModeGeometry*)GEditorModeTools().GetActiveMode(FBuiltinEditorModes::EM_Geometry);
+	FEdModeGeometry* mode = (FEdModeGeometry*)GLevelEditorModeTools().GetActiveMode(FBuiltinEditorModes::EM_Geometry);
 	return mode->GetSelectionState() ? false : true;
 }
 
@@ -2011,12 +2011,12 @@ void UGeomModifier_Clip::ApplyClip( bool InSplit, bool InFlipNormal )
 		GEditor->NoteSelectionChange();
 	}
 
-	FEdModeGeometry* Mode = (FEdModeGeometry*)GEditorModeTools().GetActiveMode(FBuiltinEditorModes::EM_Geometry);
+	FEdModeGeometry* Mode = (FEdModeGeometry*)GLevelEditorModeTools().GetActiveMode(FBuiltinEditorModes::EM_Geometry);
 	Mode->FinalizeSourceData();
 	Mode->GetFromSource();
 }
 
-bool UGeomModifier_Clip::InputKey(FLevelEditorViewportClient* ViewportClient, FViewport* Viewport, FKey Key, EInputEvent Event)
+bool UGeomModifier_Clip::InputKey(FEditorViewportClient* ViewportClient, FViewport* Viewport, FKey Key, EInputEvent Event)
 {
 	bool bResult = false;
 
@@ -2086,7 +2086,7 @@ bool UGeomModifier_Clip::InputKey(FLevelEditorViewportClient* ViewportClient, FV
 
 void UGeomModifier_Clip::Render(const FSceneView* View,FViewport* Viewport,FPrimitiveDrawInterface* PDI)
 {
-	FEdModeGeometry* mode = (FEdModeGeometry*)GEditorModeTools().GetActiveMode(FBuiltinEditorModes::EM_Geometry);
+	FEdModeGeometry* mode = (FEdModeGeometry*)GLevelEditorModeTools().GetActiveMode(FBuiltinEditorModes::EM_Geometry);
 	FModeTool_GeometryModify* tool = (FModeTool_GeometryModify*)mode->GetCurrentTool();
 	if( tool->GetCurrentModifier() != this )
 	{
@@ -2095,7 +2095,7 @@ void UGeomModifier_Clip::Render(const FSceneView* View,FViewport* Viewport,FPrim
 
 	// Only draw in ortho viewports
 
-	if( !((FLevelEditorViewportClient*)(Viewport->GetClient()))->IsOrtho() )
+	if( !((FEditorViewportClient*)(Viewport->GetClient()))->IsOrtho() )
 	{
 		return;
 	}
@@ -2134,7 +2134,7 @@ void UGeomModifier_Clip::Render(const FSceneView* View,FViewport* Viewport,FPrim
 
 		vtx3 = vtx1;
 
-		const FLevelEditorViewportClient* ViewportClient = static_cast<FLevelEditorViewportClient*>( Viewport->GetClient() );
+		const FEditorViewportClient* ViewportClient = static_cast<FEditorViewportClient*>( Viewport->GetClient() );
 		switch( ViewportClient->ViewportType )
 		{
 			case LVT_OrthoXY:
@@ -2175,11 +2175,11 @@ void UGeomModifier_Clip::Render(const FSceneView* View,FViewport* Viewport,FPrim
 	DrawWireBox(PDI, FBox::BuildAABB( SnappedMouseWorldSpacePos, FVector(BoxSz,BoxSz,BoxSz) ), FLinearColor(1,1,1), SDPG_Foreground);
 }
 
-void UGeomModifier_Clip::DrawHUD(FLevelEditorViewportClient* ViewportClient,FViewport* Viewport,const FSceneView* View,FCanvas* Canvas)
+void UGeomModifier_Clip::DrawHUD(FEditorViewportClient* ViewportClient,FViewport* Viewport,const FSceneView* View,FCanvas* Canvas)
 {
 }
 
-void UGeomModifier_Clip::Tick(FLevelEditorViewportClient* ViewportClient,float DeltaTime)
+void UGeomModifier_Clip::Tick(FEditorViewportClient* ViewportClient,float DeltaTime)
 {
 	if( GCurrentLevelEditingViewportClient == ViewportClient )
 	{
@@ -2206,14 +2206,14 @@ UGeomModifier_Delete::UGeomModifier_Delete(const class FPostConstructInitializeP
 
 bool UGeomModifier_Delete::Supports()
 {
-	FEdModeGeometry* mode = (FEdModeGeometry*)GEditorModeTools().GetActiveMode(FBuiltinEditorModes::EM_Geometry);
+	FEdModeGeometry* mode = (FEdModeGeometry*)GLevelEditorModeTools().GetActiveMode(FBuiltinEditorModes::EM_Geometry);
 	return (mode->HavePolygonsSelected() || mode->HaveVerticesSelected());
 }
 
 
 bool UGeomModifier_Delete::OnApply()
 {
-	FEdModeGeometry* mode = (FEdModeGeometry*)GEditorModeTools().GetActiveMode(FBuiltinEditorModes::EM_Geometry);
+	FEdModeGeometry* mode = (FEdModeGeometry*)GLevelEditorModeTools().GetActiveMode(FBuiltinEditorModes::EM_Geometry);
 	bool bHandled = false;
 
 	for( FEdModeGeometry::TGeomObjectIterator Itor( mode->GeomObjectItor() ) ; Itor ; ++Itor )
@@ -2274,7 +2274,7 @@ bool UGeomModifier_Delete::OnApply()
 
 	if(SelectedActor)
 	{
-		FEditorModeTools& Tools = GEditorModeTools();
+		FEditorModeTools& Tools = GLevelEditorModeTools();
 		Tools.SetPivotLocation( SelectedActor->GetActorLocation() , false );
 	}
 	
@@ -2290,13 +2290,13 @@ UGeomModifier_Create::UGeomModifier_Create(const class FPostConstructInitializeP
 
 bool UGeomModifier_Create::Supports()
 {
-	FEdModeGeometry* mode = (FEdModeGeometry*)GEditorModeTools().GetActiveMode(FBuiltinEditorModes::EM_Geometry);
+	FEdModeGeometry* mode = (FEdModeGeometry*)GLevelEditorModeTools().GetActiveMode(FBuiltinEditorModes::EM_Geometry);
 	return mode->HaveVerticesSelected();
 }
 
 bool UGeomModifier_Create::OnApply()
 {
-	FEdModeGeometry* mode = (FEdModeGeometry*)GEditorModeTools().GetActiveMode(FBuiltinEditorModes::EM_Geometry);
+	FEdModeGeometry* mode = (FEdModeGeometry*)GLevelEditorModeTools().GetActiveMode(FBuiltinEditorModes::EM_Geometry);
 
 	for( FEdModeGeometry::TGeomObjectIterator Itor( mode->GeomObjectItor() ) ; Itor ; ++Itor )
 	{
@@ -2356,13 +2356,13 @@ bool UGeomModifier_Flip::Supports()
 {
 	// Supports polygons selected and objects selected
 
-	FEdModeGeometry* mode = (FEdModeGeometry*)GEditorModeTools().GetActiveMode(FBuiltinEditorModes::EM_Geometry);
+	FEdModeGeometry* mode = (FEdModeGeometry*)GLevelEditorModeTools().GetActiveMode(FBuiltinEditorModes::EM_Geometry);
 	return (!mode->HaveEdgesSelected() && !mode->HaveVerticesSelected());
 }
 
 bool UGeomModifier_Flip::OnApply()
 {
-	FEdModeGeometry* mode = (FEdModeGeometry*)GEditorModeTools().GetActiveMode(FBuiltinEditorModes::EM_Geometry);
+	FEdModeGeometry* mode = (FEdModeGeometry*)GLevelEditorModeTools().GetActiveMode(FBuiltinEditorModes::EM_Geometry);
 	bool bHavePolygonsSelected = mode->HavePolygonsSelected();
 
 	for( FEdModeGeometry::TGeomObjectIterator Itor( mode->GeomObjectItor() ) ; Itor ; ++Itor )
@@ -2398,7 +2398,7 @@ UGeomModifier_Split::UGeomModifier_Split(const class FPostConstructInitializePro
 
 bool UGeomModifier_Split::Supports()
 {
-	FEdModeGeometry* mode = (FEdModeGeometry*)GEditorModeTools().GetActiveMode(FBuiltinEditorModes::EM_Geometry);
+	FEdModeGeometry* mode = (FEdModeGeometry*)GLevelEditorModeTools().GetActiveMode(FBuiltinEditorModes::EM_Geometry);
 
 	// This modifier assumes that a single geometry object is selected
 
@@ -2425,7 +2425,7 @@ bool UGeomModifier_Split::Supports()
 
 bool UGeomModifier_Split::OnApply()
 {
-	FEdModeGeometry* mode = (FEdModeGeometry*)GEditorModeTools().GetActiveMode(FBuiltinEditorModes::EM_Geometry);
+	FEdModeGeometry* mode = (FEdModeGeometry*)GLevelEditorModeTools().GetActiveMode(FBuiltinEditorModes::EM_Geometry);
 
 	// Get a pointer to the selected geom object
 
@@ -2768,13 +2768,13 @@ UGeomModifier_Triangulate::UGeomModifier_Triangulate(const class FPostConstructI
 
 bool UGeomModifier_Triangulate::Supports()
 {
-	FEdModeGeometry* mode = (FEdModeGeometry*)GEditorModeTools().GetActiveMode(FBuiltinEditorModes::EM_Geometry);
+	FEdModeGeometry* mode = (FEdModeGeometry*)GLevelEditorModeTools().GetActiveMode(FBuiltinEditorModes::EM_Geometry);
 	return (!mode->HaveEdgesSelected() && !mode->HaveVerticesSelected());
 }
 
 bool UGeomModifier_Triangulate::OnApply()
 {
-	FEdModeGeometry* mode = (FEdModeGeometry*)GEditorModeTools().GetActiveMode(FBuiltinEditorModes::EM_Geometry);
+	FEdModeGeometry* mode = (FEdModeGeometry*)GLevelEditorModeTools().GetActiveMode(FBuiltinEditorModes::EM_Geometry);
 	bool bHavePolygonsSelected = mode->HavePolygonsSelected();
 
 	// Mark the selected polygons so we can find them in the next loop, and create
@@ -2839,14 +2839,14 @@ UGeomModifier_Optimize::UGeomModifier_Optimize(const class FPostConstructInitial
 
 bool UGeomModifier_Optimize::Supports()
 {
-	FEdModeGeometry* mode = (FEdModeGeometry*)GEditorModeTools().GetActiveMode(FBuiltinEditorModes::EM_Geometry);
+	FEdModeGeometry* mode = (FEdModeGeometry*)GLevelEditorModeTools().GetActiveMode(FBuiltinEditorModes::EM_Geometry);
 	return !mode->HaveVerticesSelected() && !mode->HaveEdgesSelected();
 }
 
 
 bool UGeomModifier_Optimize::OnApply()
 {
-	FEdModeGeometry* mode = (FEdModeGeometry*)GEditorModeTools().GetActiveMode(FBuiltinEditorModes::EM_Geometry);
+	FEdModeGeometry* mode = (FEdModeGeometry*)GLevelEditorModeTools().GetActiveMode(FBuiltinEditorModes::EM_Geometry);
 
 	TArray<FPoly> Polygons;
 
@@ -2926,13 +2926,13 @@ UGeomModifier_Turn::UGeomModifier_Turn(const class FPostConstructInitializePrope
 
 bool UGeomModifier_Turn::Supports()
 {
-	FEdModeGeometry* mode = (FEdModeGeometry*)GEditorModeTools().GetActiveMode(FBuiltinEditorModes::EM_Geometry);
+	FEdModeGeometry* mode = (FEdModeGeometry*)GLevelEditorModeTools().GetActiveMode(FBuiltinEditorModes::EM_Geometry);
 	return mode->HaveEdgesSelected();
 }
 
 bool UGeomModifier_Turn::OnApply()
 {
-	FEdModeGeometry* mode = (FEdModeGeometry*)GEditorModeTools().GetActiveMode(FBuiltinEditorModes::EM_Geometry);
+	FEdModeGeometry* mode = (FEdModeGeometry*)GLevelEditorModeTools().GetActiveMode(FBuiltinEditorModes::EM_Geometry);
 
 	// Edges
 
@@ -3087,13 +3087,13 @@ UGeomModifier_Weld::UGeomModifier_Weld(const class FPostConstructInitializePrope
 
 bool UGeomModifier_Weld::Supports()
 {
-	FEdModeGeometry* mode = (FEdModeGeometry*)GEditorModeTools().GetActiveMode(FBuiltinEditorModes::EM_Geometry);
+	FEdModeGeometry* mode = (FEdModeGeometry*)GLevelEditorModeTools().GetActiveMode(FBuiltinEditorModes::EM_Geometry);
 	return (mode->HaveVerticesSelected() && !mode->HaveEdgesSelected() && !mode->HavePolygonsSelected());
 }
 
 bool UGeomModifier_Weld::OnApply()
 {
-	FEdModeGeometry* mode = (FEdModeGeometry*)GEditorModeTools().GetActiveMode(FBuiltinEditorModes::EM_Geometry);
+	FEdModeGeometry* mode = (FEdModeGeometry*)GLevelEditorModeTools().GetActiveMode(FBuiltinEditorModes::EM_Geometry);
 
 	// Verts
 
@@ -3146,7 +3146,7 @@ bool UGeomModifier_Weld::OnApply()
 		int32 res = go->SetPivotFromSelectionArray( Actor->SavedSelections );
 		if( res == INDEX_NONE )
 		{
-			FEditorModeTools& Tools = GEditorModeTools();
+			FEditorModeTools& Tools = GLevelEditorModeTools();
 			Tools.SetPivotLocation( Actor->GetActorLocation() , false );
 		}
 		go->ForceLastSelectionIndex( res );
diff --git a/Engine/Source/Editor/GeometryMode/Public/GeometryEdMode.h b/Engine/Source/Editor/GeometryMode/Public/GeometryEdMode.h
--- a/Engine/Source/Editor/GeometryMode/Public/GeometryEdMode.h
+++ b/Engine/Source/Editor/GeometryMode/Public/GeometryEdMode.h
@@ -8,8 +8,6 @@
  */
 class FGeometryModeModule : public IModuleInterface
 {
-private:
-	TSharedPtr<class FEdModeGeometry> EdModeGeometry;
 public:
 	// IModuleInterface
 	virtual void StartupModule() override;
@@ -25,6 +23,9 @@ public:
 class FEdModeGeometry : public FEdMode
 {
 public:
+
+	FEdModeGeometry();
+
 	/**
 	 * Struct for cacheing of selected objects components midpoints for reselection when rebuilding the BSP
 	 */
@@ -55,6 +56,7 @@ public:
 	virtual void MapChangeNotify() override;
 	virtual void SelectionChanged() override;
 	virtual FVector GetWidgetLocation() const;
+	virtual bool IsCompatibleWith(FEditorModeID OtherModeID) const override;
 	// End of FEdMode interface
 
 	// FGCObject interface
@@ -180,11 +182,6 @@ public:
 	const FGeomObject* GetGeomObject(int32 Index) const		{ return GeomObjects[ Index ]; }
 	//@}
 
-
-private:
-	FEdModeGeometry();
-
-
 protected:
 	/** 
 	 * Custom data compiled when this mode is entered, based on currently
@@ -217,7 +214,7 @@ public:
 	/**
 	 * @return		true if the delta was handled by this editor mode tool.
 	 */
-	virtual bool InputDelta(FLevelEditorViewportClient* InViewportClient,FViewport* InViewport,FVector& InDrag,FRotator& InRot,FVector& InScale);
+	virtual bool InputDelta(FEditorViewportClient* InViewportClient,FViewport* InViewport,FVector& InDrag,FRotator& InRot,FVector& InScale);
 
 	virtual bool StartModify();
 	virtual bool EndModify();
@@ -242,13 +239,13 @@ public:
 	const UGeomModifier* GetModifier(int32 Index) const		{ return Modifiers[ Index ]; }
 	//@}
 
-	virtual void Tick(FLevelEditorViewportClient* ViewportClient,float DeltaTime);
+	virtual void Tick(FEditorViewportClient* ViewportClient,float DeltaTime);
 
 	/** @return		true if the key was handled by this editor mode tool. */
-	virtual bool InputKey(FLevelEditorViewportClient* ViewportClient, FViewport* Viewport, FKey Key, EInputEvent Event) override;
+	virtual bool InputKey(FEditorViewportClient* ViewportClient, FViewport* Viewport, FKey Key, EInputEvent Event) override;
 
 	virtual void Render(const FSceneView* View,FViewport* Viewport,FPrimitiveDrawInterface* PDI);
-	virtual void DrawHUD(FLevelEditorViewportClient* ViewportClient,FViewport* Viewport,const FSceneView* View,FCanvas* Canvas);
+	virtual void DrawHUD(FEditorViewportClient* ViewportClient,FViewport* Viewport,const FSceneView* View,FCanvas* Canvas);
 	
 	/**
 	 * Store the current geom selections for all geom objects
diff --git a/Engine/Source/Editor/LandscapeEditor/Private/ActorFactoryLandscape.cpp b/Engine/Source/Editor/LandscapeEditor/Private/ActorFactoryLandscape.cpp
--- a/Engine/Source/Editor/LandscapeEditor/Private/ActorFactoryLandscape.cpp
+++ b/Engine/Source/Editor/LandscapeEditor/Private/ActorFactoryLandscape.cpp
@@ -15,9 +15,9 @@ UActorFactoryLandscape::UActorFactoryLandscape(const class FPostConstructInitial
 
 AActor* UActorFactoryLandscape::SpawnActor(UObject* Asset, ULevel* InLevel, const FVector& Location, const FRotator& Rotation, EObjectFlags ObjectFlags, const FName& Name)
 {
-	GEditorModeTools().ActivateMode(FBuiltinEditorModes::EM_Landscape);
+	GLevelEditorModeTools().ActivateMode(FBuiltinEditorModes::EM_Landscape);
 
-	FEdModeLandscape* EdMode = (FEdModeLandscape*)GEditorModeTools().GetActiveMode(FBuiltinEditorModes::EM_Landscape);
+	FEdModeLandscape* EdMode = (FEdModeLandscape*)GLevelEditorModeTools().GetActiveMode(FBuiltinEditorModes::EM_Landscape);
 
 	EdMode->UISettings->NewLandscape_Location = Location;
 	EdMode->UISettings->NewLandscape_Rotation = Rotation;
@@ -66,9 +66,9 @@ bool ALandscapePlaceholder::TeleportTo(const FVector& DestLocation, const FRotat
 {
 	bool bResult = Super::TeleportTo(DestLocation, DestRotation, bIsATest, bNoCheck);
 
-	GEditorModeTools().ActivateMode(FBuiltinEditorModes::EM_Landscape);
+	GLevelEditorModeTools().ActivateMode(FBuiltinEditorModes::EM_Landscape);
 
-	FEdModeLandscape* EdMode = (FEdModeLandscape*)GEditorModeTools().GetActiveMode(FBuiltinEditorModes::EM_Landscape);
+	FEdModeLandscape* EdMode = (FEdModeLandscape*)GLevelEditorModeTools().GetActiveMode(FBuiltinEditorModes::EM_Landscape);
 
 	EdMode->UISettings->NewLandscape_Location = GetActorLocation();
 	EdMode->UISettings->NewLandscape_Rotation = GetActorRotation();
diff --git a/Engine/Source/Editor/LandscapeEditor/Private/LandscapeEdMode.cpp b/Engine/Source/Editor/LandscapeEditor/Private/LandscapeEdMode.cpp
--- a/Engine/Source/Editor/LandscapeEditor/Private/LandscapeEdMode.cpp
+++ b/Engine/Source/Editor/LandscapeEditor/Private/LandscapeEdMode.cpp
@@ -159,12 +159,6 @@ FEdModeLandscape::FEdModeLandscape()
 ,	DraggingEdge_Remainder(0)
 ,	CachedLandscapeMaterial(NULL)
 {
-	ID = FBuiltinEditorModes::EM_Landscape;
-	Name = NSLOCTEXT("EditorModes", "LandscapeMode", "Landscape");
-	IconBrush = FSlateIcon(FEditorStyle::GetStyleSetName(), "LevelEditor.LandscapeMode", "LevelEditor.LandscapeMode.Small");
-	bVisible = true;
-	PriorityOrder = 300;
-
 	GizmoMaterial = LoadObject<UMaterial>(NULL, TEXT("/Engine/EditorLandscapeResources/GizmoMaterial.GizmoMaterial"), NULL, LOAD_None, NULL);
 
 	GLayerDebugColorMaterial = LoadObject<UMaterial>(NULL, TEXT("/Engine/EditorLandscapeResources/LayerVisMaterial.LayerVisMaterial"), NULL, LOAD_None, NULL);
@@ -569,7 +563,7 @@ void FEdModeLandscape::Exit()
 }
 
 /** FEdMode: Called once per frame */
-void FEdModeLandscape::Tick(FLevelEditorViewportClient* ViewportClient,float DeltaTime)
+void FEdModeLandscape::Tick(FEditorViewportClient* ViewportClient,float DeltaTime)
 {
 	FEdMode::Tick(ViewportClient,DeltaTime);
 
@@ -615,7 +609,7 @@ void FEdModeLandscape::Tick(FLevelEditorViewportClient* ViewportClient,float Del
 
 
 /** FEdMode: Called when the mouse is moved over the viewport */
-bool FEdModeLandscape::MouseMove( FLevelEditorViewportClient* ViewportClient, FViewport* Viewport, int32 MouseX, int32 MouseY )
+bool FEdModeLandscape::MouseMove( FEditorViewportClient* ViewportClient, FViewport* Viewport, int32 MouseX, int32 MouseY )
 {
 	if (bToolActive && !Viewport->KeyState(EKeys::LeftMouseButton))
 	{
@@ -659,7 +653,7 @@ bool FEdModeLandscape::DisallowMouseDeltaTracking() const
  *
  * @return	true if input was handled
  */
-bool FEdModeLandscape::CapturedMouseMove( FLevelEditorViewportClient* ViewportClient, FViewport* Viewport, int32 MouseX, int32 MouseY )
+bool FEdModeLandscape::CapturedMouseMove( FEditorViewportClient* ViewportClient, FViewport* Viewport, int32 MouseX, int32 MouseY )
 {
 	return MouseMove(ViewportClient, Viewport, MouseX, MouseY);
 }
@@ -670,7 +664,7 @@ namespace
 }
 
 /** FEdMode: Called when a mouse button is pressed */
-bool FEdModeLandscape::StartTracking(FLevelEditorViewportClient* InViewportClient, FViewport* InViewport)
+bool FEdModeLandscape::StartTracking(FEditorViewportClient* InViewportClient, FViewport* InViewport)
 {
 	if (CurrentGizmoActor.IsValid() && CurrentGizmoActor->IsSelected() && GLandscapeEditRenderMode & ELandscapeEditRenderMode::Gizmo)
 	{
@@ -683,7 +677,7 @@ bool FEdModeLandscape::StartTracking(FLevelEditorViewportClient* InViewportClien
 
 
 /** FEdMode: Called when the a mouse button is released */
-bool FEdModeLandscape::EndTracking(FLevelEditorViewportClient* InViewportClient, FViewport* InViewport)
+bool FEdModeLandscape::EndTracking(FEditorViewportClient* InViewportClient, FViewport* InViewport)
 {
 	if (GIsGizmoDragging)
 	{
@@ -717,7 +711,7 @@ namespace
 };
 
 /** Trace under the mouse cursor and return the landscape hit and the hit location (in landscape quad space) */
-bool FEdModeLandscape::LandscapeMouseTrace( FLevelEditorViewportClient* ViewportClient, float& OutHitX, float& OutHitY )
+bool FEdModeLandscape::LandscapeMouseTrace( FEditorViewportClient* ViewportClient, float& OutHitX, float& OutHitY )
 {
 	int32 MouseX = ViewportClient->Viewport->GetMouseX();
 	int32 MouseY = ViewportClient->Viewport->GetMouseY();
@@ -725,7 +719,7 @@ bool FEdModeLandscape::LandscapeMouseTrace( FLevelEditorViewportClient* Viewport
 	return LandscapeMouseTrace( ViewportClient, MouseX, MouseY, OutHitX, OutHitY );
 }
 
-bool FEdModeLandscape::LandscapeMouseTrace( FLevelEditorViewportClient* ViewportClient, FVector& OutHitLocation )
+bool FEdModeLandscape::LandscapeMouseTrace( FEditorViewportClient* ViewportClient, FVector& OutHitLocation )
 {
 	int32 MouseX = ViewportClient->Viewport->GetMouseX();
 	int32 MouseY = ViewportClient->Viewport->GetMouseY();
@@ -734,7 +728,7 @@ bool FEdModeLandscape::LandscapeMouseTrace( FLevelEditorViewportClient* Viewport
 }
 
 /** Trace under the specified coordinates and return the landscape hit and the hit location (in landscape quad space) */
-bool FEdModeLandscape::LandscapeMouseTrace( FLevelEditorViewportClient* ViewportClient, int32 MouseX, int32 MouseY, float& OutHitX, float& OutHitY )
+bool FEdModeLandscape::LandscapeMouseTrace( FEditorViewportClient* ViewportClient, int32 MouseX, int32 MouseY, float& OutHitX, float& OutHitY )
 {
 	FVector HitLocation;
 	bool bResult = LandscapeMouseTrace( ViewportClient, MouseX, MouseY, HitLocation );
@@ -743,7 +737,7 @@ bool FEdModeLandscape::LandscapeMouseTrace( FLevelEditorViewportClient* Viewport
 	return bResult;
 }
 
-bool FEdModeLandscape::LandscapeMouseTrace( FLevelEditorViewportClient* ViewportClient, int32 MouseX, int32 MouseY, FVector& OutHitLocation )
+bool FEdModeLandscape::LandscapeMouseTrace( FEditorViewportClient* ViewportClient, int32 MouseX, int32 MouseY, FVector& OutHitLocation )
 {
 	// Cache a copy of the world pointer	
 	UWorld* World = ViewportClient->GetWorld();
@@ -823,7 +817,7 @@ bool FEdModeLandscape::LandscapeMouseTrace( FLevelEditorViewportClient* Viewport
 	return false;
 }
 
-bool FEdModeLandscape::LandscapePlaneTrace(FLevelEditorViewportClient* ViewportClient, const FPlane& Plane, FVector& OutHitLocation)
+bool FEdModeLandscape::LandscapePlaneTrace(FEditorViewportClient* ViewportClient, const FPlane& Plane, FVector& OutHitLocation)
 {
 	int32 MouseX = ViewportClient->Viewport->GetMouseX();
 	int32 MouseY = ViewportClient->Viewport->GetMouseY();
@@ -831,7 +825,7 @@ bool FEdModeLandscape::LandscapePlaneTrace(FLevelEditorViewportClient* ViewportC
 	return LandscapePlaneTrace(ViewportClient, MouseX, MouseY, Plane, OutHitLocation);
 }
 
-bool FEdModeLandscape::LandscapePlaneTrace(FLevelEditorViewportClient* ViewportClient, int32 MouseX, int32 MouseY, const FPlane& Plane, FVector& OutHitLocation)
+bool FEdModeLandscape::LandscapePlaneTrace(FEditorViewportClient* ViewportClient, int32 MouseX, int32 MouseY, const FPlane& Plane, FVector& OutHitLocation)
 {
 	// Cache a copy of the world pointer
 	UWorld* World = ViewportClient->GetWorld();
@@ -1130,7 +1124,7 @@ bool FEdModeLandscape::ProcessEditPaste()
 	return Result;
 }
 
-bool FEdModeLandscape::HandleClick(FLevelEditorViewportClient* InViewportClient, HHitProxy* HitProxy, const FViewportClick& Click)
+bool FEdModeLandscape::HandleClick(FEditorViewportClient* InViewportClient, HHitProxy* HitProxy, const FViewportClick& Click)
 {
 	if (GEditor->PlayWorld != NULL)
 	{
@@ -1152,7 +1146,7 @@ bool FEdModeLandscape::HandleClick(FLevelEditorViewportClient* InViewportClient,
 }
 
 /** FEdMode: Called when a key is pressed */
-bool FEdModeLandscape::InputKey( FLevelEditorViewportClient* ViewportClient, FViewport* Viewport, FKey Key, EInputEvent Event )
+bool FEdModeLandscape::InputKey( FEditorViewportClient* ViewportClient, FViewport* Viewport, FKey Key, EInputEvent Event )
 {
 	if (GEditor->PlayWorld != NULL)
 	{
@@ -1178,7 +1172,7 @@ bool FEdModeLandscape::InputKey( FLevelEditorViewportClient* ViewportClient, FVi
 						DraggingEdge = EdgeProxy->Edge;
 						DraggingEdge_Remainder = 0;
 
-						return false; // false to let FLevelEditorViewportClient.InputKey start mouse tracking and enable InputDelta() so we can use it
+						return false; // false to let FEditorViewportClient.InputKey start mouse tracking and enable InputDelta() so we can use it
 					}
 				}
 			}
@@ -1189,7 +1183,7 @@ bool FEdModeLandscape::InputKey( FLevelEditorViewportClient* ViewportClient, FVi
 					DraggingEdge = ELandscapeEdge::None;
 					DraggingEdge_Remainder = 0;
 
-					return false; // false to let FLevelEditorViewportClient.InputKey end mouse tracking
+					return false; // false to let FEditorViewportClient.InputKey end mouse tracking
 				}
 			}
 		}
@@ -1351,7 +1345,7 @@ bool FEdModeLandscape::InputKey( FLevelEditorViewportClient* ViewportClient, FVi
 }
 
 /** FEdMode: Called when mouse drag input is applied */
-bool FEdModeLandscape::InputDelta( FLevelEditorViewportClient* InViewportClient, FViewport* InViewport, FVector& InDrag, FRotator& InRot, FVector& InScale )
+bool FEdModeLandscape::InputDelta( FEditorViewportClient* InViewportClient, FViewport* InViewport, FVector& InDrag, FRotator& InRot, FVector& InScale )
 {
 	if (GEditor->PlayWorld != NULL)
 	{
@@ -1808,7 +1802,7 @@ void FEdModeLandscape::OnWorldChange()
 	if (NewLandscapePreviewMode == ENewLandscapePreviewMode::None &&
 		CurrentToolTarget.LandscapeInfo == NULL)
 	{
-		GEditorModeTools().DeactivateMode(ID);
+		RequestDeletion();
 	}
 }
 
@@ -1844,7 +1838,7 @@ void FEdModeLandscape::Render( const FSceneView* View, FViewport* Viewport, FPri
 		static const FLinearColor SectionBorderColour   (0.0f,  0.4f,  0.0f);
 		static const FLinearColor InnerColour           (0.0f,  0.25f, 0.0f);
 
-		const ELevelViewportType ViewportType = ((FLevelEditorViewportClient*)Viewport->GetClient())->ViewportType;
+		const ELevelViewportType ViewportType = ((FEditorViewportClient*)Viewport->GetClient())->ViewportType;
 
 		const int32 ComponentCountX = UISettings->NewLandscape_ComponentCount.X;
 		const int32 ComponentCountY = UISettings->NewLandscape_ComponentCount.Y;
@@ -2086,7 +2080,7 @@ void FEdModeLandscape::Render( const FSceneView* View, FViewport* Viewport, FPri
 }
 
 /** FEdMode: Render HUD elements for this tool */
-void FEdModeLandscape::DrawHUD( FLevelEditorViewportClient* ViewportClient, FViewport* Viewport, const FSceneView* View, FCanvas* Canvas )
+void FEdModeLandscape::DrawHUD( FEditorViewportClient* ViewportClient, FViewport* Viewport, const FSceneView* View, FCanvas* Canvas )
 {
 
 }
@@ -2295,7 +2289,7 @@ void FEdModeLandscape::ForceRealTimeViewports(const bool bEnable, const bool bSt
 		{
 			if (ViewportWindow.IsValid())
 			{
-				FLevelEditorViewportClient& Viewport = ViewportWindow->GetLevelViewportClient();
+				FEditorViewportClient& Viewport = ViewportWindow->GetLevelViewportClient();
 				if (bEnable)
 				{
 					Viewport.SetRealtime(bEnable, bStoreCurrentState);
diff --git a/Engine/Source/Editor/LandscapeEditor/Private/LandscapeEdModeBrushes.cpp b/Engine/Source/Editor/LandscapeEditor/Private/LandscapeEdModeBrushes.cpp
--- a/Engine/Source/Editor/LandscapeEditor/Private/LandscapeEdModeBrushes.cpp
+++ b/Engine/Source/Editor/LandscapeEditor/Private/LandscapeEdModeBrushes.cpp
@@ -99,7 +99,7 @@ public:
 		LastMousePosition = FVector2D(LandscapeX, LandscapeY);
 	}
 
-	virtual void Tick(FLevelEditorViewportClient* ViewportClient,float DeltaTime) override
+	virtual void Tick(FEditorViewportClient* ViewportClient,float DeltaTime) override
 	{
 		ULandscapeInfo* LandscapeInfo = EdMode->CurrentToolTarget.LandscapeInfo.Get();
 		ALandscapeProxy* Proxy = LandscapeInfo ? LandscapeInfo->GetLandscapeProxy() : NULL;
@@ -205,7 +205,7 @@ public:
 		LastMousePosition = FVector2D(LandscapeX, LandscapeY);
 	}
 
-	virtual bool InputKey( FLevelEditorViewportClient* InViewportClient, FViewport* InViewport, FKey InKey, EInputEvent InEvent ) override
+	virtual bool InputKey( FEditorViewportClient* InViewportClient, FViewport* InViewport, FKey InKey, EInputEvent InEvent ) override
 	{
 		bool bUpdate = false;
 		return bUpdate;
@@ -341,7 +341,7 @@ public:
 		LastMousePosition = FVector2D(LandscapeX, LandscapeY);
 	}
 
-	virtual void Tick(FLevelEditorViewportClient* ViewportClient,float DeltaTime) override
+	virtual void Tick(FEditorViewportClient* ViewportClient,float DeltaTime) override
 	{
 		TSet<ULandscapeComponent*> NewComponents;
 
@@ -408,7 +408,7 @@ public:
 		LastMousePosition = FVector2D(LandscapeX, LandscapeY);
 	}
 
-	virtual bool InputKey( FLevelEditorViewportClient* InViewportClient, FViewport* InViewport, FKey InKey, EInputEvent InEvent ) override
+	virtual bool InputKey( FEditorViewportClient* InViewportClient, FViewport* InViewport, FKey InKey, EInputEvent InEvent ) override
 	{
 		bool bUpdate = false;
 
@@ -535,7 +535,7 @@ public:
 		LastMousePosition = FVector2D(LandscapeX, LandscapeY);
 	}
 
-	virtual void Tick(FLevelEditorViewportClient* ViewportClient,float DeltaTime) override
+	virtual void Tick(FEditorViewportClient* ViewportClient,float DeltaTime) override
 	{
 		if (GLandscapeEditRenderMode & ELandscapeEditRenderMode::Gizmo || GLandscapeEditRenderMode & ELandscapeEditRenderMode::Select)
 		{
@@ -611,7 +611,7 @@ public:
 		LastMousePosition = FVector2D(LandscapeX, LandscapeY);
 	}
 
-	virtual bool InputKey( FLevelEditorViewportClient* InViewportClient, FViewport* InViewport, FKey InKey, EInputEvent InEvent ) override
+	virtual bool InputKey( FEditorViewportClient* InViewportClient, FViewport* InViewport, FKey InKey, EInputEvent InEvent ) override
 	{
 		bool bUpdate = false;
 		return bUpdate;
@@ -1014,7 +1014,7 @@ public:
 		return (X1 <= X2 && Y1 <= Y2);
 	}
 
-	virtual void Tick(FLevelEditorViewportClient* ViewportClient,float DeltaTime) override
+	virtual void Tick(FEditorViewportClient* ViewportClient,float DeltaTime) override
 	{
 		FLandscapeBrushCircle::Tick(ViewportClient,DeltaTime);
 
@@ -1171,7 +1171,7 @@ public:
 		}
 	}
 
-	virtual void Tick(FLevelEditorViewportClient* ViewportClient,float DeltaTime) override
+	virtual void Tick(FEditorViewportClient* ViewportClient,float DeltaTime) override
 	{
 		FLandscapeBrushCircle::Tick(ViewportClient,DeltaTime);
 
diff --git a/Engine/Source/Editor/LandscapeEditor/Private/LandscapeEdModeComponentTools.cpp b/Engine/Source/Editor/LandscapeEditor/Private/LandscapeEdModeComponentTools.cpp
--- a/Engine/Source/Editor/LandscapeEditor/Private/LandscapeEdModeComponentTools.cpp
+++ b/Engine/Source/Editor/LandscapeEditor/Private/LandscapeEdModeComponentTools.cpp
@@ -30,7 +30,7 @@ public:
 	{
 	}
 
-	virtual void Apply(FLevelEditorViewportClient* ViewportClient, FLandscapeBrush* Brush, const ULandscapeEditorObject* UISettings, const TArray<FLandscapeToolMousePosition>& MousePositions)
+	virtual void Apply(FEditorViewportClient* ViewportClient, FLandscapeBrush* Brush, const ULandscapeEditorObject* UISettings, const TArray<FLandscapeToolMousePosition>& MousePositions)
 	{
 		if( LandscapeInfo )
 		{
@@ -202,7 +202,7 @@ public:
 		,	Cache(InTarget)
 	{}
 
-	virtual void Apply(FLevelEditorViewportClient* ViewportClient, FLandscapeBrush* Brush, const ULandscapeEditorObject* UISettings, const TArray<FLandscapeToolMousePosition>& MousePositions)
+	virtual void Apply(FEditorViewportClient* ViewportClient, FLandscapeBrush* Brush, const ULandscapeEditorObject* UISettings, const TArray<FLandscapeToolMousePosition>& MousePositions)
 	{
 		if( LandscapeInfo )
 		{
@@ -290,7 +290,7 @@ public:
 		: LandscapeInfo(InTarget.LandscapeInfo.Get())
 	{}
 
-	virtual void Apply(FLevelEditorViewportClient* ViewportClient, FLandscapeBrush* Brush, const ULandscapeEditorObject* UISettings, const TArray<FLandscapeToolMousePosition>& MousePositions)
+	virtual void Apply(FEditorViewportClient* ViewportClient, FLandscapeBrush* Brush, const ULandscapeEditorObject* UISettings, const TArray<FLandscapeToolMousePosition>& MousePositions)
 	{
 		ALandscape* Landscape = LandscapeInfo ? LandscapeInfo->LandscapeActor.Get() : NULL;
 
@@ -757,7 +757,7 @@ public:
 		,	HeightCache(InTarget)
 	{}
 
-	virtual void Apply(FLevelEditorViewportClient* ViewportClient, FLandscapeBrush* Brush, const ULandscapeEditorObject* UISettings, const TArray<FLandscapeToolMousePosition>& MousePositions)
+	virtual void Apply(FEditorViewportClient* ViewportClient, FLandscapeBrush* Brush, const ULandscapeEditorObject* UISettings, const TArray<FLandscapeToolMousePosition>& MousePositions)
 	{
 		ALandscapeProxy* Landscape = LandscapeInfo ? LandscapeInfo->GetCurrentLevelLandscapeProxy(true) : NULL;
 		if( Landscape && EdMode->LandscapeRenderAddCollision)
@@ -885,7 +885,7 @@ public:
 		:	LandscapeInfo(InTarget.LandscapeInfo.Get())
 	{}
 
-	virtual void Apply(FLevelEditorViewportClient* ViewportClient, FLandscapeBrush* Brush, const ULandscapeEditorObject* UISettings, const TArray<FLandscapeToolMousePosition>& MousePositions)
+	virtual void Apply(FEditorViewportClient* ViewportClient, FLandscapeBrush* Brush, const ULandscapeEditorObject* UISettings, const TArray<FLandscapeToolMousePosition>& MousePositions)
 	{
 		if( LandscapeInfo )
 		{
@@ -1075,7 +1075,7 @@ public:
 		float Data;
 	};
 
-	virtual void Apply(FLevelEditorViewportClient* ViewportClient, FLandscapeBrush* Brush, const ULandscapeEditorObject* UISettings, const TArray<FLandscapeToolMousePosition>& MousePositions)
+	virtual void Apply(FEditorViewportClient* ViewportClient, FLandscapeBrush* Brush, const ULandscapeEditorObject* UISettings, const TArray<FLandscapeToolMousePosition>& MousePositions)
 	{
 		//ULandscapeInfo* LandscapeInfo = EdMode->CurrentToolTarget.LandscapeInfo;
 		ALandscapeGizmoActiveActor* Gizmo = EdMode->CurrentGizmoActor.Get();
@@ -1384,7 +1384,7 @@ public:
 		return Target.TargetType == ToolTarget::TargetType;
 	}
 
-	virtual bool BeginTool( FLevelEditorViewportClient* ViewportClient, const FLandscapeToolTarget& InTarget, const FVector& InHitLocation ) override
+	virtual bool BeginTool( FEditorViewportClient* ViewportClient, const FLandscapeToolTarget& InTarget, const FVector& InHitLocation ) override
 	{
 		check(this->MousePositions.Num()==0);
 		this->bToolActive = true;
@@ -1399,7 +1399,7 @@ public:
 		return true;
 	}
 
-	virtual void EndTool(FLevelEditorViewportClient* ViewportClient) override 
+	virtual void EndTool(FEditorViewportClient* ViewportClient) override 
 	{
 		if( this->bToolActive && this->MousePositions.Num() )
 		{
@@ -1413,7 +1413,7 @@ public:
 		this->EdMode->GizmoBrush->EndStroke();
 	}
 
-	virtual bool MouseMove( FLevelEditorViewportClient* ViewportClient, FViewport* Viewport, int x, int y ) override
+	virtual bool MouseMove( FEditorViewportClient* ViewportClient, FViewport* Viewport, int x, int y ) override
 	{
 		return true;
 	}	
@@ -1432,7 +1432,7 @@ public:
 		,	WeightCache(InTarget)
 	{}
 
-	virtual void Apply(FLevelEditorViewportClient* ViewportClient, FLandscapeBrush* Brush, const ULandscapeEditorObject* UISettings, const TArray<FLandscapeToolMousePosition>& MousePositions)
+	virtual void Apply(FEditorViewportClient* ViewportClient, FLandscapeBrush* Brush, const ULandscapeEditorObject* UISettings, const TArray<FLandscapeToolMousePosition>& MousePositions)
 	{
 		//ULandscapeInfo* LandscapeInfo = EdMode->CurrentToolTarget.LandscapeInfo;
 		ALandscapeGizmoActiveActor* Gizmo = EdMode->CurrentGizmoActor.Get();
@@ -1687,7 +1687,7 @@ public:
 		return Target.TargetType == ToolTarget::TargetType;
 	}
 
-	virtual bool BeginTool( FLevelEditorViewportClient* ViewportClient, const FLandscapeToolTarget& InTarget, const FVector& InHitLocation ) override
+	virtual bool BeginTool( FEditorViewportClient* ViewportClient, const FLandscapeToolTarget& InTarget, const FVector& InHitLocation ) override
 	{
 		this->bToolActive = true;
 		this->ToolStroke = new FLandscapeToolStrokePaste<ToolTarget>(this->EdMode, InTarget);
@@ -1708,7 +1708,7 @@ public:
 		return true;
 	}
 
-	virtual void EndTool(FLevelEditorViewportClient* ViewportClient) override
+	virtual void EndTool(FEditorViewportClient* ViewportClient) override
 	{
 		if( this->bToolActive && this->MousePositions.Num() )
 		{
@@ -1729,7 +1729,7 @@ public:
 		}
 	}
 
-	virtual bool MouseMove( FLevelEditorViewportClient* ViewportClient, FViewport* Viewport, int32 x, int32 y ) override
+	virtual bool MouseMove( FEditorViewportClient* ViewportClient, FViewport* Viewport, int32 x, int32 y ) override
 	{
 		if (bUseGizmoRegion)
 		{
@@ -1831,18 +1831,18 @@ public:
 		EdMode->NewLandscapePreviewMode = ENewLandscapePreviewMode::None;
 	}
 
-	virtual bool BeginTool( FLevelEditorViewportClient* ViewportClient, const FLandscapeToolTarget& Target, const FVector& InHitLocation )
+	virtual bool BeginTool( FEditorViewportClient* ViewportClient, const FLandscapeToolTarget& Target, const FVector& InHitLocation )
 	{
 		// does nothing
 		return false;
 	}
 
-	virtual void EndTool(FLevelEditorViewportClient* ViewportClient)
+	virtual void EndTool(FEditorViewportClient* ViewportClient)
 	{
 		// does nothing
 	}
 
-	virtual bool MouseMove( FLevelEditorViewportClient* ViewportClient, FViewport* Viewport, int32 x, int32 y )
+	virtual bool MouseMove( FEditorViewportClient* ViewportClient, FViewport* Viewport, int32 x, int32 y )
 	{
 		// does nothing
 		return false;
@@ -1900,18 +1900,18 @@ public:
 	{
 	}
 
-	virtual bool BeginTool( FLevelEditorViewportClient* ViewportClient, const FLandscapeToolTarget& Target, const FVector& InHitLocation )
+	virtual bool BeginTool( FEditorViewportClient* ViewportClient, const FLandscapeToolTarget& Target, const FVector& InHitLocation )
 	{
 		// does nothing
 		return false;
 	}
 
-	virtual void EndTool(FLevelEditorViewportClient* ViewportClient)
+	virtual void EndTool(FEditorViewportClient* ViewportClient)
 	{
 		// does nothing
 	}
 
-	virtual bool MouseMove( FLevelEditorViewportClient* ViewportClient, FViewport* Viewport, int32 x, int32 y )
+	virtual bool MouseMove( FEditorViewportClient* ViewportClient, FViewport* Viewport, int32 x, int32 y )
 	{
 		// does nothing
 		return false;
diff --git a/Engine/Source/Editor/LandscapeEditor/Private/LandscapeEdModeErosionTools.cpp b/Engine/Source/Editor/LandscapeEditor/Private/LandscapeEdModeErosionTools.cpp
--- a/Engine/Source/Editor/LandscapeEditor/Private/LandscapeEdModeErosionTools.cpp
+++ b/Engine/Source/Editor/LandscapeEditor/Private/LandscapeEdModeErosionTools.cpp
@@ -26,7 +26,7 @@ public:
 	,	bWeightApplied(InTarget.TargetType != ELandscapeToolTargetType::Heightmap)
 	{}
 
-	virtual void Apply(FLevelEditorViewportClient* ViewportClient, FLandscapeBrush* Brush, const ULandscapeEditorObject* UISettings, const TArray<FLandscapeToolMousePosition>& MousePositions) = 0;
+	virtual void Apply(FEditorViewportClient* ViewportClient, FLandscapeBrush* Brush, const ULandscapeEditorObject* UISettings, const TArray<FLandscapeToolMousePosition>& MousePositions) = 0;
 protected:
 	class ULandscapeInfo* LandscapeInfo;
 	FLandscapeHeightCache HeightCache;
@@ -60,7 +60,7 @@ public:
 	:	FLandscapeToolStrokeErosionBase(InEdMode, InTarget)
 	{}
 
-	virtual void Apply(FLevelEditorViewportClient* ViewportClient, FLandscapeBrush* Brush, const ULandscapeEditorObject* UISettings, const TArray<FLandscapeToolMousePosition>& MousePositions) override
+	virtual void Apply(FEditorViewportClient* ViewportClient, FLandscapeBrush* Brush, const ULandscapeEditorObject* UISettings, const TArray<FLandscapeToolMousePosition>& MousePositions) override
 	{
 		if( !LandscapeInfo )
 		{
@@ -276,7 +276,7 @@ public:
 	:	FLandscapeToolStrokeErosionBase(InEdMode, InTarget)
 	{}
 
-	virtual void Apply(FLevelEditorViewportClient* ViewportClient, FLandscapeBrush* Brush, const ULandscapeEditorObject* UISettings, const TArray<FLandscapeToolMousePosition>& MousePositions) override
+	virtual void Apply(FEditorViewportClient* ViewportClient, FLandscapeBrush* Brush, const ULandscapeEditorObject* UISettings, const TArray<FLandscapeToolMousePosition>& MousePositions) override
 	{
 		if( !LandscapeInfo )
 		{
diff --git a/Engine/Source/Editor/LandscapeEditor/Private/LandscapeEdModePaintTools.cpp b/Engine/Source/Editor/LandscapeEditor/Private/LandscapeEdModePaintTools.cpp
--- a/Engine/Source/Editor/LandscapeEditor/Private/LandscapeEdModePaintTools.cpp
+++ b/Engine/Source/Editor/LandscapeEditor/Private/LandscapeEdModePaintTools.cpp
@@ -62,7 +62,7 @@ public:
 	{
 	}
 
-	virtual void Apply(FLevelEditorViewportClient* ViewportClient, FLandscapeBrush* Brush, const ULandscapeEditorObject* UISettings, const TArray<FLandscapeToolMousePosition>& MousePositions) = 0;
+	virtual void Apply(FEditorViewportClient* ViewportClient, FLandscapeBrush* Brush, const ULandscapeEditorObject* UISettings, const TArray<FLandscapeToolMousePosition>& MousePositions) = 0;
 protected:
 	typename ToolTarget::CacheClass Cache;
 	class ULandscapeInfo* LandscapeInfo;
@@ -85,7 +85,7 @@ public:
 	:	FLandscapeToolStrokePaintBase<ToolTarget>(InEdMode, InTarget)
 	{}
 
-	virtual void Apply(FLevelEditorViewportClient* ViewportClient, FLandscapeBrush* Brush, const ULandscapeEditorObject* UISettings, const TArray<FLandscapeToolMousePosition>& MousePositions) override
+	virtual void Apply(FEditorViewportClient* ViewportClient, FLandscapeBrush* Brush, const ULandscapeEditorObject* UISettings, const TArray<FLandscapeToolMousePosition>& MousePositions) override
 	{
 		// Get list of verts to update
 		TMap<FIntPoint, float> BrushInfo;
@@ -321,7 +321,7 @@ public:
 	:	FLandscapeToolStrokePaintBase<ToolTarget>(InEdMode, InTarget)
 	{}
 
-	virtual void Apply(FLevelEditorViewportClient* ViewportClient, FLandscapeBrush* Brush, const ULandscapeEditorObject* UISettings, const TArray<FLandscapeToolMousePosition>& MousePositions) override
+	virtual void Apply(FEditorViewportClient* ViewportClient, FLandscapeBrush* Brush, const ULandscapeEditorObject* UISettings, const TArray<FLandscapeToolMousePosition>& MousePositions) override
 	{
 		if (!this->LandscapeInfo) return;
 
@@ -436,7 +436,7 @@ public:
 		}
 	}
 
-	virtual void Apply(FLevelEditorViewportClient* ViewportClient, FLandscapeBrush* Brush, const ULandscapeEditorObject* UISettings, const TArray<FLandscapeToolMousePosition>& MousePositions) override
+	virtual void Apply(FEditorViewportClient* ViewportClient, FLandscapeBrush* Brush, const ULandscapeEditorObject* UISettings, const TArray<FLandscapeToolMousePosition>& MousePositions) override
 	{
 		if (!this->LandscapeInfo) return;
 
@@ -572,7 +572,7 @@ public:
 	virtual const TCHAR* GetToolName() override { return TEXT("Flatten"); }
 	virtual FText GetDisplayName() override { return NSLOCTEXT("UnrealEd", "LandscapeMode_Flatten", "Flatten"); };
 
-	virtual void Tick(FLevelEditorViewportClient* ViewportClient,float DeltaTime)
+	virtual void Tick(FEditorViewportClient* ViewportClient,float DeltaTime)
 	{
 		FLandscapeToolPaintBase<ToolTarget, FLandscapeToolStrokeFlatten<ToolTarget>>::Tick(ViewportClient, DeltaTime);
 
@@ -580,13 +580,13 @@ public:
 		MeshComponent->SetVisibility(bShowGrid);
 	}
 
-	virtual bool MouseMove(FLevelEditorViewportClient* ViewportClient, FViewport* Viewport, int32 x, int32 y) override
+	virtual bool MouseMove(FEditorViewportClient* ViewportClient, FViewport* Viewport, int32 x, int32 y) override
 	{
 		bool bResult = FLandscapeToolPaintBase<ToolTarget, FLandscapeToolStrokeFlatten<ToolTarget>>::MouseMove(ViewportClient, Viewport, x, y);
 
 		if (ViewportClient->IsLevelEditorClient() && MeshComponent != NULL)
 		{
-			this->EdMode->LandscapeMouseTrace((FLevelEditorViewportClient*)ViewportClient, x, y, LastMousePosition);
+			this->EdMode->LandscapeMouseTrace((FEditorViewportClient*)ViewportClient, x, y, LastMousePosition);
 
 			const FTransform LocalToWorld = this->EdMode->CurrentToolTarget.LandscapeInfo->GetLandscapeProxy()->ActorToWorld();
 			FVector Origin;
@@ -639,7 +639,7 @@ public:
 	:	FLandscapeToolStrokePaintBase<ToolTarget>(InEdMode, InTarget)
 	{}
 
-	virtual void Apply(FLevelEditorViewportClient* ViewportClient, FLandscapeBrush* Brush, const ULandscapeEditorObject* UISettings, const TArray<FLandscapeToolMousePosition>& MousePositions) override
+	virtual void Apply(FEditorViewportClient* ViewportClient, FLandscapeBrush* Brush, const ULandscapeEditorObject* UISettings, const TArray<FLandscapeToolMousePosition>& MousePositions) override
 	{
 		if (!this->LandscapeInfo) return;
 
diff --git a/Engine/Source/Editor/LandscapeEditor/Private/LandscapeEdModeRampTool.cpp b/Engine/Source/Editor/LandscapeEditor/Private/LandscapeEdModeRampTool.cpp
--- a/Engine/Source/Editor/LandscapeEditor/Private/LandscapeEdModeRampTool.cpp
+++ b/Engine/Source/Editor/LandscapeEditor/Private/LandscapeEdModeRampTool.cpp
@@ -114,10 +114,10 @@ public:
 	{
 		NumPoints = 0;
 		SelectedPoint = INDEX_NONE;
-		GEditorModeTools().SetWidgetMode(FWidget::WM_Translate);
+		GLevelEditorModeTools().SetWidgetMode(FWidget::WM_Translate);
 	}
 
-	virtual bool BeginTool(FLevelEditorViewportClient* ViewportClient, const FLandscapeToolTarget& Target, const FVector& InHitLocation) override
+	virtual bool BeginTool(FEditorViewportClient* ViewportClient, const FLandscapeToolTarget& Target, const FVector& InHitLocation) override
 	{
 		if (NumPoints < 2)
 		{
@@ -125,7 +125,7 @@ public:
 			SelectedPoint = NumPoints;
 			NumPoints++;
 			bMovingPoint = true;
-			GEditorModeTools().SetWidgetMode(FWidget::WM_Translate);
+			GLevelEditorModeTools().SetWidgetMode(FWidget::WM_Translate);
 		}
 		else
 		{
@@ -133,7 +133,7 @@ public:
 			{
 				Points[SelectedPoint] = InHitLocation;
 				bMovingPoint = true;
-				GEditorModeTools().SetWidgetMode(FWidget::WM_Translate);
+				GLevelEditorModeTools().SetWidgetMode(FWidget::WM_Translate);
 			}
 		}
 
@@ -142,12 +142,12 @@ public:
 		return true;
 	}
 
-	virtual void EndTool(FLevelEditorViewportClient* ViewportClient) override
+	virtual void EndTool(FEditorViewportClient* ViewportClient) override
 	{
 		bMovingPoint = false;
 	}
 
-	virtual bool MouseMove(FLevelEditorViewportClient* ViewportClient, FViewport* Viewport, int32 x, int32 y) override
+	virtual bool MouseMove(FEditorViewportClient* ViewportClient, FViewport* Viewport, int32 x, int32 y) override
 	{
 		if (bMovingPoint)
 		{
@@ -175,7 +175,7 @@ public:
 			{
 				HLandscapeRampToolPointHitProxy* PointHitProxy = (HLandscapeRampToolPointHitProxy*)HitProxy;
 				SelectedPoint = PointHitProxy->Point;
-				GEditorModeTools().SetWidgetMode(FWidget::WM_Translate);
+				GLevelEditorModeTools().SetWidgetMode(FWidget::WM_Translate);
 				GUnrealEd->RedrawLevelEditingViewports();
 
 				return true;
@@ -185,7 +185,7 @@ public:
 		return false;
 	}
 
-	virtual bool InputKey(FLevelEditorViewportClient* InViewportClient, FViewport* InViewport, FKey InKey, EInputEvent InEvent) override
+	virtual bool InputKey(FEditorViewportClient* InViewportClient, FViewport* InViewport, FKey InKey, EInputEvent InEvent) override
 	{
 		if (InKey == EKeys::Enter && InEvent == IE_Pressed)
 		{
@@ -267,7 +267,7 @@ public:
 		return false;
 	}
 
-	virtual bool InputDelta(FLevelEditorViewportClient* InViewportClient, FViewport* InViewport, FVector& InDrag, FRotator& InRot, FVector& InScale) override
+	virtual bool InputDelta(FEditorViewportClient* InViewportClient, FViewport* InViewport, FVector& InDrag, FRotator& InRot, FVector& InScale) override
 	{
 		if (SelectedPoint != INDEX_NONE && InViewportClient->GetCurrentWidgetAxis() != EAxisList::None)
 		{
diff --git a/Engine/Source/Editor/LandscapeEditor/Private/LandscapeEdModeSplineTools.cpp b/Engine/Source/Editor/LandscapeEditor/Private/LandscapeEdModeSplineTools.cpp
--- a/Engine/Source/Editor/LandscapeEditor/Private/LandscapeEdModeSplineTools.cpp
+++ b/Engine/Source/Editor/LandscapeEditor/Private/LandscapeEdModeSplineTools.cpp
@@ -75,7 +75,7 @@ public:
 
 	void UpdatePropertiesWindows()
 	{
-		if( GEditorModeTools().IsModeActive( EdMode->GetID() ) )
+		if( GLevelEditorModeTools().IsModeActive( EdMode->GetID() ) )
 		{
 			TArray<UObject*> Objects;
 			Objects.Reset(SelectedSplineControlPoints.Num() + SelectedSplineSegments.Num());
@@ -154,7 +154,7 @@ public:
 		Segment->Modify();
 		Segment->SetSplineSelected(true);
 
-		GEditorModeTools().SetWidgetMode(FWidget::WM_Scale);
+		GLevelEditorModeTools().SetWidgetMode(FWidget::WM_Scale);
 	}
 
 	void SelectConnected()
@@ -622,7 +622,7 @@ public:
 		}
 	}
 
-	virtual bool BeginTool( FLevelEditorViewportClient* ViewportClient, const FLandscapeToolTarget& InTarget, const FVector& InHitLocation) override
+	virtual bool BeginTool( FEditorViewportClient* ViewportClient, const FLandscapeToolTarget& InTarget, const FVector& InHitLocation) override
 	{
 		LandscapeInfo = InTarget.LandscapeInfo.Get();
 
@@ -646,12 +646,12 @@ public:
 		return true;
 	}
 
-	virtual void EndTool(FLevelEditorViewportClient* ViewportClient) override
+	virtual void EndTool(FEditorViewportClient* ViewportClient) override
 	{
 		LandscapeInfo = NULL;
 	}
 
-	virtual bool MouseMove( FLevelEditorViewportClient* ViewportClient, FViewport* Viewport, int32 x, int32 y ) override
+	virtual bool MouseMove( FEditorViewportClient* ViewportClient, FViewport* Viewport, int32 x, int32 y ) override
 	{
 		FVector HitLocation;
 		if (EdMode->LandscapeMouseTrace(ViewportClient, x, y, HitLocation))
@@ -666,7 +666,7 @@ public:
 		return true;
 	}
 
-	virtual void ApplyTool( FLevelEditorViewportClient* ViewportClient )
+	virtual void ApplyTool( FEditorViewportClient* ViewportClient )
 	{
 	}
 
@@ -765,7 +765,7 @@ public:
 		return false;
 	}
 
-	virtual bool InputKey( FLevelEditorViewportClient* InViewportClient, FViewport* InViewport, FKey InKey, EInputEvent InEvent ) override
+	virtual bool InputKey( FEditorViewportClient* InViewportClient, FViewport* InViewport, FKey InKey, EInputEvent InEvent ) override
 	{
 		if (InKey == EKeys::F4 && InEvent == IE_Pressed)
 		{
@@ -1026,7 +1026,7 @@ public:
 						SplinesComponent->Modify();
 						DraggingTangent_Segment->Modify();
 
-						return false; // false to let FLevelEditorViewportClient.InputKey start mouse tracking and enable InputDelta() so we can use it
+						return false; // false to let FEditorViewportClient.InputKey start mouse tracking and enable InputDelta() so we can use it
 					}
 				}
 			}
@@ -1053,7 +1053,7 @@ public:
 
 					GEditor->EndTransaction();
 
-					return false; // false to let FLevelEditorViewportClient.InputKey end mouse tracking
+					return false; // false to let FEditorViewportClient.InputKey end mouse tracking
 				}
 			}
 		}
@@ -1061,7 +1061,7 @@ public:
 		return false;
 	}
 
-	virtual bool InputDelta( FLevelEditorViewportClient* InViewportClient, FViewport* InViewport, FVector& InDrag, FRotator& InRot, FVector& InScale ) override
+	virtual bool InputDelta( FEditorViewportClient* InViewportClient, FViewport* InViewport, FVector& InDrag, FRotator& InRot, FVector& InScale ) override
 	{
 		FVector Drag = InDrag;
 
@@ -1267,7 +1267,7 @@ public:
 			FVector HandlePos1 = SplinesComponent->ComponentToWorld.TransformPosition(ControlPoint->Location + ControlPoint->Rotation.Vector() * 20);
 			DrawDashedLine(PDI, HandlePos0, HandlePos1, FColor::White, 20, SDPG_Foreground);
 
-			if (GEditorModeTools().GetWidgetMode() == FWidget::WM_Scale)
+			if (GLevelEditorModeTools().GetWidgetMode() == FWidget::WM_Scale)
 			{
 				for (const FLandscapeSplineConnection& Connection : ControlPoint->ConnectedSegments)
 				{
@@ -1285,7 +1285,7 @@ public:
 			}
 		}
 
-		if (GEditorModeTools().GetWidgetMode() == FWidget::WM_Scale)
+		if (GLevelEditorModeTools().GetWidgetMode() == FWidget::WM_Scale)
 		{
 			for (ULandscapeSplineSegment* Segment : SelectedSplineSegments)
 			{
diff --git a/Engine/Source/Editor/LandscapeEditor/Private/LandscapeEdModeTools.h b/Engine/Source/Editor/LandscapeEditor/Private/LandscapeEdModeTools.h
--- a/Engine/Source/Editor/LandscapeEditor/Private/LandscapeEdModeTools.h
+++ b/Engine/Source/Editor/LandscapeEditor/Private/LandscapeEdModeTools.h
@@ -1138,7 +1138,7 @@ class FLandscapeStrokeBase
 {
 public:
 	FLandscapeStrokeBase(TInputType& InTarget) {}
-	virtual void Apply(FLevelEditorViewportClient* ViewportClient, FLandscapeBrush* Brush, const ULandscapeEditorObject* UISettings, const TArray<FLandscapeToolMousePosition>& MousePositions) = 0;
+	virtual void Apply(FEditorViewportClient* ViewportClient, FLandscapeBrush* Brush, const ULandscapeEditorObject* UISettings, const TArray<FLandscapeToolMousePosition>& MousePositions) = 0;
 };
 
 
@@ -1157,7 +1157,7 @@ public:
 	,	ToolStroke(NULL)
 	{}
 
-	virtual bool BeginTool( FLevelEditorViewportClient* ViewportClient, const FLandscapeToolTarget& InTarget, const FVector& InHitLocation ) override
+	virtual bool BeginTool( FEditorViewportClient* ViewportClient, const FLandscapeToolTarget& InTarget, const FVector& InHitLocation ) override
 	{
 		if( !ensure(MousePositions.Num()==0) ) 
 		{
@@ -1175,7 +1175,7 @@ public:
 		return true;
 	}
 
-	virtual void Tick(FLevelEditorViewportClient* ViewportClient,float DeltaTime) override 
+	virtual void Tick(FEditorViewportClient* ViewportClient,float DeltaTime) override 
 	{
 		if( bToolActive && MousePositions.Num() )
 		{
@@ -1184,7 +1184,7 @@ public:
 		}
 	}
 
-	virtual void EndTool(FLevelEditorViewportClient* ViewportClient) override 
+	virtual void EndTool(FEditorViewportClient* ViewportClient) override 
 	{
 		if( bToolActive && MousePositions.Num() )
 		{
@@ -1198,7 +1198,7 @@ public:
 		EdMode->CurrentBrush->EndStroke();
 	}
 
-	virtual bool MouseMove( FLevelEditorViewportClient* ViewportClient, FViewport* Viewport, int32 x, int32 y ) override 
+	virtual bool MouseMove( FEditorViewportClient* ViewportClient, FViewport* Viewport, int32 x, int32 y ) override 
 	{
 		FVector HitLocation;
 		if( EdMode->LandscapeMouseTrace(ViewportClient, x, y, HitLocation)  )
diff --git a/Engine/Source/Editor/LandscapeEditor/Private/LandscapeEdModeXYOffset.cpp b/Engine/Source/Editor/LandscapeEditor/Private/LandscapeEdModeXYOffset.cpp
--- a/Engine/Source/Editor/LandscapeEditor/Private/LandscapeEdModeXYOffset.cpp
+++ b/Engine/Source/Editor/LandscapeEditor/Private/LandscapeEdModeXYOffset.cpp
@@ -58,7 +58,7 @@ public:
 		,	Cache(InTarget)
 	{}
 
-	virtual void Apply(FLevelEditorViewportClient* ViewportClient, FLandscapeBrush* Brush, const ULandscapeEditorObject* UISettings, const TArray<FLandscapeToolMousePosition>& MousePositions)
+	virtual void Apply(FEditorViewportClient* ViewportClient, FLandscapeBrush* Brush, const ULandscapeEditorObject* UISettings, const TArray<FLandscapeToolMousePosition>& MousePositions)
 	{
 		if (!LandscapeInfo) 
 		{
diff --git a/Engine/Source/Editor/LandscapeEditor/Private/LandscapeEditorDetailCustomization_Base.cpp b/Engine/Source/Editor/LandscapeEditor/Private/LandscapeEditorDetailCustomization_Base.cpp
--- a/Engine/Source/Editor/LandscapeEditor/Private/LandscapeEditorDetailCustomization_Base.cpp
+++ b/Engine/Source/Editor/LandscapeEditor/Private/LandscapeEditorDetailCustomization_Base.cpp
@@ -12,7 +12,7 @@
 
 class FEdModeLandscape* FLandscapeEditorDetailCustomization_Base::GetEditorMode()
 {
-	return (FEdModeLandscape*)GEditorModeTools().GetActiveMode(FBuiltinEditorModes::EM_Landscape);
+	return (FEdModeLandscape*)GLevelEditorModeTools().GetActiveMode(FBuiltinEditorModes::EM_Landscape);
 }
 
 bool FLandscapeEditorDetailCustomization_Base::IsToolActive(FName ToolSetName)
@@ -80,5 +80,5 @@ void FLandscapeEditorDetailCustomization_Base::OnValueCommitted(int32 NewValue,
 
 class FEdModeLandscape* FLandscapeEditorStructCustomization_Base::GetEditorMode()
 {
-	return (FEdModeLandscape*)GEditorModeTools().GetActiveMode(FBuiltinEditorModes::EM_Landscape);
+	return (FEdModeLandscape*)GLevelEditorModeTools().GetActiveMode(FBuiltinEditorModes::EM_Landscape);
 }
diff --git a/Engine/Source/Editor/LandscapeEditor/Private/LandscapeEditorDetailCustomization_TargetLayers.cpp b/Engine/Source/Editor/LandscapeEditor/Private/LandscapeEditorDetailCustomization_TargetLayers.cpp
--- a/Engine/Source/Editor/LandscapeEditor/Private/LandscapeEditorDetailCustomization_TargetLayers.cpp
+++ b/Engine/Source/Editor/LandscapeEditor/Private/LandscapeEditorDetailCustomization_TargetLayers.cpp
@@ -74,7 +74,7 @@ bool FLandscapeEditorDetailCustomization_TargetLayers::ShouldShowTargetLayers()
 
 FEdModeLandscape* FLandscapeEditorCustomNodeBuilder_TargetLayers::GetEditorMode()
 {
-	return (FEdModeLandscape*)GEditorModeTools().GetActiveMode(FBuiltinEditorModes::EM_Landscape);
+	return (FEdModeLandscape*)GLevelEditorModeTools().GetActiveMode(FBuiltinEditorModes::EM_Landscape);
 }
 
 FLandscapeEditorCustomNodeBuilder_TargetLayers::FLandscapeEditorCustomNodeBuilder_TargetLayers(TSharedRef<FAssetThumbnailPool> InThumbnailPool)
diff --git a/Engine/Source/Editor/LandscapeEditor/Private/LandscapeEditorModule.cpp b/Engine/Source/Editor/LandscapeEditor/Private/LandscapeEditorModule.cpp
--- a/Engine/Source/Editor/LandscapeEditor/Private/LandscapeEditorModule.cpp
+++ b/Engine/Source/Editor/LandscapeEditor/Private/LandscapeEditorModule.cpp
@@ -20,8 +20,6 @@
 
 class FLandscapeEditorModule : public ILandscapeEditorModule
 {
-private:
-	TSharedPtr<FEdModeLandscape> EdModeLandscape;
 public:
 	
 	/**
@@ -32,9 +30,13 @@ public:
 		FLandscapeEditorCommands::Register();
 
 		// register the editor mode
-		TSharedRef<FEdModeLandscape> NewEditorMode = MakeShareable(new FEdModeLandscape);
-		GEditorModeTools().RegisterMode(NewEditorMode);
-		EdModeLandscape = NewEditorMode;
+		FEditorModeRegistry::Get().RegisterMode<FEdModeLandscape>(
+			FBuiltinEditorModes::EM_Landscape,
+			NSLOCTEXT("EditorModes", "LandscapeMode", "Landscape"),
+			FSlateIcon(FEditorStyle::GetStyleSetName(), "LevelEditor.LandscapeMode", "LevelEditor.LandscapeMode.Small"),
+			true,
+			300
+			);
 
 		FPropertyEditorModule& PropertyModule = FModuleManager::GetModuleChecked<FPropertyEditorModule>("PropertyEditor");
 		PropertyModule.RegisterCustomClassLayout("LandscapeEditorObject", FOnGetDetailCustomizationInstance::CreateStatic(&FLandscapeEditorDetails::MakeInstance));
@@ -71,8 +73,7 @@ public:
 		FLandscapeEditorCommands::Unregister();
 
 		// unregister the editor mode
-		GEditorModeTools().UnregisterMode(EdModeLandscape.ToSharedRef());
-		EdModeLandscape = NULL;
+		FEditorModeRegistry::Get().UnregisterMode(FBuiltinEditorModes::EM_Landscape);
 	}
 
 	static void ConstructLandscapeViewportMenu(FMenuBuilder& MenuBuilder)
diff --git a/Engine/Source/Editor/LandscapeEditor/Private/LandscapeEditorObject.cpp b/Engine/Source/Editor/LandscapeEditor/Private/LandscapeEditorObject.cpp
--- a/Engine/Source/Editor/LandscapeEditor/Private/LandscapeEditorObject.cpp
+++ b/Engine/Source/Editor/LandscapeEditor/Private/LandscapeEditorObject.cpp
@@ -362,7 +362,7 @@ void ULandscapeEditorObject::SetbSnapGizmo(bool InbSnapGizmo)
 {
 	bSnapGizmo = InbSnapGizmo;
 
-	FEdModeLandscape* EdMode = (FEdModeLandscape*)GEditorModeTools().GetActiveMode(FBuiltinEditorModes::EM_Landscape);
+	FEdModeLandscape* EdMode = (FEdModeLandscape*)GLevelEditorModeTools().GetActiveMode(FBuiltinEditorModes::EM_Landscape);
 	if (EdMode && EdMode->CurrentGizmoActor.IsValid())
 	{
 		EdMode->CurrentGizmoActor->bSnapToLandscapeGrid = bSnapGizmo;
diff --git a/Engine/Source/Editor/LandscapeEditor/Private/SLandscapeEditor.cpp b/Engine/Source/Editor/LandscapeEditor/Private/SLandscapeEditor.cpp
--- a/Engine/Source/Editor/LandscapeEditor/Private/SLandscapeEditor.cpp
+++ b/Engine/Source/Editor/LandscapeEditor/Private/SLandscapeEditor.cpp
@@ -86,7 +86,7 @@ FText FLandscapeToolKit::GetBaseToolkitName() const
 
 class FEdModeLandscape* FLandscapeToolKit::GetEditorMode() const
 {
-	return (FEdModeLandscape*)GEditorModeTools().GetActiveMode(FBuiltinEditorModes::EM_Landscape);
+	return (FEdModeLandscape*)GLevelEditorModeTools().GetActiveMode(FBuiltinEditorModes::EM_Landscape);
 }
 
 TSharedPtr<SWidget> FLandscapeToolKit::GetInlineContent() const
@@ -177,7 +177,7 @@ END_SLATE_FUNCTION_BUILD_OPTIMIZATION
 
 class FEdModeLandscape* SLandscapeEditor::GetEditorMode() const
 {
-	return (FEdModeLandscape*)GEditorModeTools().GetActiveMode(FBuiltinEditorModes::EM_Landscape);
+	return (FEdModeLandscape*)GLevelEditorModeTools().GetActiveMode(FBuiltinEditorModes::EM_Landscape);
 }
 
 void SLandscapeEditor::UpdateErrorText() const
diff --git a/Engine/Source/Editor/LandscapeEditor/Public/LandscapeEdMode.h b/Engine/Source/Editor/LandscapeEditor/Public/LandscapeEdMode.h
--- a/Engine/Source/Editor/LandscapeEditor/Public/LandscapeEdMode.h
+++ b/Engine/Source/Editor/LandscapeEditor/Public/LandscapeEdMode.h
@@ -40,8 +40,8 @@ public:
 	
 	virtual void MouseMove( float LandscapeX, float LandscapeY ) = 0;
 	virtual bool ApplyBrush( const TArray<FLandscapeToolMousePosition>& MousePositions, TMap<FIntPoint, float>& OutBrush, int32& OutX1, int32& OutY1, int32& OutX2, int32& OutY2 ) =0;
-	virtual bool InputKey( FLevelEditorViewportClient* InViewportClient, FViewport* InViewport, FKey InKey, EInputEvent InEvent ) { return false; }
-	virtual void Tick(FLevelEditorViewportClient* ViewportClient,float DeltaTime) {};
+	virtual bool InputKey( FEditorViewportClient* InViewportClient, FViewport* InViewport, FKey InKey, EInputEvent InEvent ) { return false; }
+	virtual void Tick(FEditorViewportClient* ViewportClient,float DeltaTime) {};
 	virtual void BeginStroke(float LandscapeX, float LandscapeY, class FLandscapeTool* CurrentTool);
 	virtual void EndStroke();
 	virtual void EnterBrush() {}
@@ -159,13 +159,13 @@ public:
 	virtual void EnterTool() {}
 	virtual void ExitTool() {}
 	virtual bool IsValidForTarget(const FLandscapeToolTarget& Target) = 0;
-	virtual bool BeginTool( FLevelEditorViewportClient* ViewportClient, const FLandscapeToolTarget& Target, const FVector& InHitLocation ) = 0;
-	virtual void EndTool(FLevelEditorViewportClient* ViewportClient) = 0;
-	virtual void Tick(FLevelEditorViewportClient* ViewportClient,float DeltaTime) {};
-	virtual bool MouseMove( FLevelEditorViewportClient* ViewportClient, FViewport* Viewport, int32 x, int32 y ) = 0;
+	virtual bool BeginTool( FEditorViewportClient* ViewportClient, const FLandscapeToolTarget& Target, const FVector& InHitLocation ) = 0;
+	virtual void EndTool(FEditorViewportClient* ViewportClient) = 0;
+	virtual void Tick(FEditorViewportClient* ViewportClient,float DeltaTime) {};
+	virtual bool MouseMove( FEditorViewportClient* ViewportClient, FViewport* Viewport, int32 x, int32 y ) = 0;
 	virtual bool HandleClick( HHitProxy* HitProxy, const FViewportClick& Click ) { return false; }
-	virtual bool InputKey( FLevelEditorViewportClient* InViewportClient, FViewport* InViewport, FKey InKey, EInputEvent InEvent ) { return false; }
-	virtual bool InputDelta( FLevelEditorViewportClient* InViewportClient, FViewport* InViewport, FVector& InDrag, FRotator& InRot, FVector& InScale ) { return false; }
+	virtual bool InputKey( FEditorViewportClient* InViewportClient, FViewport* InViewport, FKey InKey, EInputEvent InEvent ) { return false; }
+	virtual bool InputDelta( FEditorViewportClient* InViewportClient, FViewport* InViewport, FVector& InDrag, FRotator& InRot, FVector& InScale ) { return false; }
 	virtual ~FLandscapeTool() {}
 	virtual const TCHAR* GetToolName() = 0;
 	virtual FText GetDisplayName() = 0;
@@ -533,7 +533,7 @@ public:
 	virtual void Exit() override;
 
 	/** FEdMode: Called when the mouse is moved over the viewport */
-	virtual bool MouseMove( FLevelEditorViewportClient* ViewportClient, FViewport* Viewport, int32 x, int32 y ) override;
+	virtual bool MouseMove( FEditorViewportClient* ViewportClient, FViewport* Viewport, int32 x, int32 y ) override;
 
 	/**
 	 * FEdMode: Called when the mouse is moved while a window input capture is in effect
@@ -545,34 +545,34 @@ public:
 	 *
 	 * @return	true if input was handled
 	 */
-	virtual bool CapturedMouseMove( FLevelEditorViewportClient* InViewportClient, FViewport* InViewport, int32 InMouseX, int32 InMouseY ) override;
+	virtual bool CapturedMouseMove( FEditorViewportClient* InViewportClient, FViewport* InViewport, int32 InMouseX, int32 InMouseY ) override;
 
 	/** FEdMode: Called when a mouse button is pressed */
-	virtual bool StartTracking(FLevelEditorViewportClient* InViewportClient, FViewport* InViewport) override;
+	virtual bool StartTracking(FEditorViewportClient* InViewportClient, FViewport* InViewport) override;
 
 	/** FEdMode: Called when a mouse button is released */
-	virtual bool EndTracking(FLevelEditorViewportClient* InViewportClient, FViewport* InViewport) override;
+	virtual bool EndTracking(FEditorViewportClient* InViewportClient, FViewport* InViewport) override;
 
 	/** FEdMode: Allow us to disable mouse delta tracking during painting */
 	virtual bool DisallowMouseDeltaTracking() const override;
 
 	/** FEdMode: Called once per frame */
-	virtual void Tick(FLevelEditorViewportClient* ViewportClient,float DeltaTime) override;
+	virtual void Tick(FEditorViewportClient* ViewportClient,float DeltaTime) override;
 
 	/** FEdMode: Called when clicking on a hit proxy */
-	virtual bool HandleClick(FLevelEditorViewportClient* InViewportClient, HHitProxy* HitProxy, const FViewportClick& Click) override;
+	virtual bool HandleClick(FEditorViewportClient* InViewportClient, HHitProxy* HitProxy, const FViewportClick& Click) override;
 
 	/** FEdMode: Called when a key is pressed */
-	virtual bool InputKey( FLevelEditorViewportClient* InViewportClient, FViewport* InViewport, FKey InKey, EInputEvent InEvent ) override;
+	virtual bool InputKey( FEditorViewportClient* InViewportClient, FViewport* InViewport, FKey InKey, EInputEvent InEvent ) override;
 
 	/** FEdMode: Called when mouse drag input is applied */
-	virtual bool InputDelta( FLevelEditorViewportClient* InViewportClient, FViewport* InViewport, FVector& InDrag, FRotator& InRot, FVector& InScale ) override;
+	virtual bool InputDelta( FEditorViewportClient* InViewportClient, FViewport* InViewport, FVector& InDrag, FRotator& InRot, FVector& InScale ) override;
 
 	/** FEdMode: Render elements for the landscape tool */
 	virtual void Render( const FSceneView* View, FViewport* Viewport, FPrimitiveDrawInterface* PDI ) override;
 
 	/** FEdMode: Render HUD elements for this tool */
-	virtual void DrawHUD( FLevelEditorViewportClient* ViewportClient, FViewport* Viewport, const FSceneView* View, FCanvas* Canvas ) override;
+	virtual void DrawHUD( FEditorViewportClient* ViewportClient, FViewport* Viewport, const FSceneView* View, FCanvas* Canvas ) override;
 
 	/** FEdMode: Handling SelectActor */
 	virtual bool Select( AActor* InActor, bool bInSelected ) override;
@@ -615,16 +615,16 @@ public:
 	void ForceRealTimeViewports( const bool bEnable, const bool bStoreCurrentState );
 
 	/** Trace under the mouse cursor and return the landscape hit and the hit location (in landscape quad space) */
-	bool LandscapeMouseTrace( FLevelEditorViewportClient* ViewportClient, float& OutHitX, float& OutHitY );
-	bool LandscapeMouseTrace( FLevelEditorViewportClient* ViewportClient, FVector& OutHitLocation );
+	bool LandscapeMouseTrace( FEditorViewportClient* ViewportClient, float& OutHitX, float& OutHitY );
+	bool LandscapeMouseTrace( FEditorViewportClient* ViewportClient, FVector& OutHitLocation );
 
 	/** Trace under the specified coordinates and return the landscape hit and the hit location (in landscape quad space) */
-	bool LandscapeMouseTrace( FLevelEditorViewportClient* ViewportClient, int32 MouseX, int32 MouseY, float& OutHitX, float& OutHitY );
-	bool LandscapeMouseTrace( FLevelEditorViewportClient* ViewportClient, int32 MouseX, int32 MouseY, FVector& OutHitLocation );
+	bool LandscapeMouseTrace( FEditorViewportClient* ViewportClient, int32 MouseX, int32 MouseY, float& OutHitX, float& OutHitY );
+	bool LandscapeMouseTrace( FEditorViewportClient* ViewportClient, int32 MouseX, int32 MouseY, FVector& OutHitLocation );
 
 	/** Trace under the mouse cursor / specified screen coordinates against a world-space plane and return the hit location (in world space) */
-	bool LandscapePlaneTrace(FLevelEditorViewportClient* ViewportClient, const FPlane& Plane, FVector& OutHitLocation);
-	bool LandscapePlaneTrace(FLevelEditorViewportClient* ViewportClient, int32 MouseX, int32 MouseY, const FPlane& Plane, FVector& OutHitLocation);
+	bool LandscapePlaneTrace(FEditorViewportClient* ViewportClient, const FPlane& Plane, FVector& OutHitLocation);
+	bool LandscapePlaneTrace(FEditorViewportClient* ViewportClient, int32 MouseX, int32 MouseY, const FPlane& Plane, FVector& OutHitLocation);
 
 	void SetCurrentToolMode( FName ToolModeName, bool bRestoreCurrentTool = true);
 
diff --git a/Engine/Source/Editor/LevelEditor/Private/LevelEditor.cpp b/Engine/Source/Editor/LevelEditor/Private/LevelEditor.cpp
--- a/Engine/Source/Editor/LevelEditor/Private/LevelEditor.cpp
+++ b/Engine/Source/Editor/LevelEditor/Private/LevelEditor.cpp
@@ -57,8 +57,8 @@ TSharedRef<SDockTab> FLevelEditorModule::SpawnLevelEditor( const FSpawnTabArgs&
 		SetLevelEditorInstance(LevelEditorTmp);
 		LevelEditorTmp->Initialize( LevelEditorTab, OwnerWindow.ToSharedRef() );
 
-		GEditorModeTools().DeactivateAllModes();
-		GEditorModeTools().ActivateMode(FBuiltinEditorModes::EM_Placement);
+		GLevelEditorModeTools().SetDefaultMode( FBuiltinEditorModes::EM_Placement );
+		GLevelEditorModeTools().DeactivateAllModes();
 	}
 
 	IUserFeedbackModule& UserFeedback = FModuleManager::LoadModuleChecked<IUserFeedbackModule>(TEXT("UserFeedback"));
diff --git a/Engine/Source/Editor/LevelEditor/Private/LevelEditorActions.cpp b/Engine/Source/Editor/LevelEditor/Private/LevelEditorActions.cpp
--- a/Engine/Source/Editor/LevelEditor/Private/LevelEditorActions.cpp
+++ b/Engine/Source/Editor/LevelEditor/Private/LevelEditorActions.cpp
@@ -170,7 +170,7 @@ void FLevelEditorActionCallbacks::NewLevel()
 
 bool FLevelEditorActionCallbacks::NewLevel_CanExecute()
 {
-	return FSlateApplication::Get().IsNormalExecution() && !GEditorModeTools().IsTracking();
+	return FSlateApplication::Get().IsNormalExecution() && !GLevelEditorModeTools().IsTracking();
 }
 
 void FLevelEditorActionCallbacks::OpenLevel()
@@ -180,7 +180,7 @@ void FLevelEditorActionCallbacks::OpenLevel()
 
 bool FLevelEditorActionCallbacks::OpenLevel_CanExecute()
 {
-	return FSlateApplication::Get().IsNormalExecution() && !GEditorModeTools().IsTracking();
+	return FSlateApplication::Get().IsNormalExecution() && !GLevelEditorModeTools().IsTracking();
 }
 
 FAssetPickerConfig FLevelEditorActionCallbacks::CreateLevelAssetPickerConfig()
@@ -1214,9 +1214,14 @@ void FLevelEditorActionCallbacks::AddActor_Clicked( UActorFactory* ActorFactory,
 		FLevelEditorModule& LevelEditorModule = FModuleManager::LoadModuleChecked<FLevelEditorModule>("LevelEditor");
 		LevelEditorModule.FocusViewport();
 
+		// Make sure we're in actor placement mode
+		GLevelEditorModeTools().ActivateMode(FBuiltinEditorModes::EM_Placement);
+
 		TArray<UObject*> AssetsToPlace;
 		AssetsToPlace.Add(Object);
-		IPlacementModeModule::Get().GetPlacementMode()->StartPlacing(AssetsToPlace, ActorFactory);
+
+		auto* PlacementMode = GLevelEditorModeTools().GetActiveModeTyped<IPlacementMode>(FBuiltinEditorModes::EM_Placement);
+		PlacementMode->StartPlacing(AssetsToPlace, ActorFactory);
 	}
 	else
 	{
@@ -1351,7 +1356,7 @@ void FLevelEditorActionCallbacks::ReplaceActorsFromClass_Clicked( UClass* ActorC
 bool FLevelEditorActionCallbacks::Duplicate_CanExecute()
 {
 	TArray<FEdMode*> ActiveModes; 
-	GEditorModeTools().GetActiveModes( ActiveModes );
+	GLevelEditorModeTools().GetActiveModes( ActiveModes );
 	for( int32 ModeIndex = 0; ModeIndex < ActiveModes.Num(); ++ModeIndex )
 	{
 		const EEditAction::Type CanProcess = ActiveModes[ModeIndex]->GetActionEditDuplicate();
@@ -1370,7 +1375,7 @@ bool FLevelEditorActionCallbacks::Duplicate_CanExecute()
 bool FLevelEditorActionCallbacks::Delete_CanExecute()
 {
 	TArray<FEdMode*> ActiveModes; 
-	GEditorModeTools().GetActiveModes( ActiveModes );
+	GLevelEditorModeTools().GetActiveModes( ActiveModes );
 	for( int32 ModeIndex = 0; ModeIndex < ActiveModes.Num(); ++ModeIndex )
 	{
 		const EEditAction::Type CanProcess = ActiveModes[ModeIndex]->GetActionEditDelete();
@@ -1403,7 +1408,7 @@ bool FLevelEditorActionCallbacks::Rename_CanExecute()
 bool FLevelEditorActionCallbacks::Cut_CanExecute()
 {
 	TArray<FEdMode*> ActiveModes; 
-	GEditorModeTools().GetActiveModes( ActiveModes );
+	GLevelEditorModeTools().GetActiveModes( ActiveModes );
 	for( int32 ModeIndex = 0; ModeIndex < ActiveModes.Num(); ++ModeIndex )
 	{
 		const EEditAction::Type CanProcess = ActiveModes[ModeIndex]->GetActionEditCut();
@@ -1422,7 +1427,7 @@ bool FLevelEditorActionCallbacks::Cut_CanExecute()
 bool FLevelEditorActionCallbacks::Copy_CanExecute()
 {
 	TArray<FEdMode*> ActiveModes; 
-	GEditorModeTools().GetActiveModes( ActiveModes );
+	GLevelEditorModeTools().GetActiveModes( ActiveModes );
 	for( int32 ModeIndex = 0; ModeIndex < ActiveModes.Num(); ++ModeIndex )
 	{
 		const EEditAction::Type CanProcess = ActiveModes[ModeIndex]->GetActionEditCopy();
@@ -1441,7 +1446,7 @@ bool FLevelEditorActionCallbacks::Copy_CanExecute()
 bool FLevelEditorActionCallbacks::Paste_CanExecute()
 {
 	TArray<FEdMode*> ActiveModes; 
-	GEditorModeTools().GetActiveModes( ActiveModes );
+	GLevelEditorModeTools().GetActiveModes( ActiveModes );
 	for( int32 ModeIndex = 0; ModeIndex < ActiveModes.Num(); ++ModeIndex )
 	{
 		const EEditAction::Type CanProcess = ActiveModes[ModeIndex]->GetActionEditPaste();
@@ -1487,9 +1492,9 @@ void FLevelEditorActionCallbacks::OnSelectMatineeActor( AMatineeActor * ActorToS
 
 void FLevelEditorActionCallbacks::OnSelectMatineeGroup( AActor * Actor )
 {
-	if( GEditorModeTools().IsModeActive( FBuiltinEditorModes::EM_InterpEdit ) )
+	if( GLevelEditorModeTools().IsModeActive( FBuiltinEditorModes::EM_InterpEdit ) )
 	{
-		FEdModeInterpEdit* InterpEditMode = (FEdModeInterpEdit*)GEditorModeTools().GetActiveMode( FBuiltinEditorModes::EM_InterpEdit );
+		FEdModeInterpEdit* InterpEditMode = (FEdModeInterpEdit*)GLevelEditorModeTools().GetActiveMode( FBuiltinEditorModes::EM_InterpEdit );
 
 		if ( InterpEditMode && InterpEditMode->MatineeActor )
 		{
@@ -2002,13 +2007,13 @@ void FLevelEditorActionCallbacks::OnShowOnlySelectedActors()
 
 void FLevelEditorActionCallbacks::OnToggleTransformWidgetVisibility()
 {
-	GEditorModeTools().SetShowWidget( !GEditorModeTools().GetShowWidget() );
+	GLevelEditorModeTools().SetShowWidget( !GLevelEditorModeTools().GetShowWidget() );
 	GUnrealEd->RedrawAllViewports();
 }
 
 bool FLevelEditorActionCallbacks::OnGetTransformWidgetVisibility()
 {
-	return GEditorModeTools().GetShowWidget();
+	return GLevelEditorModeTools().GetShowWidget();
 }
 
 void FLevelEditorActionCallbacks::OnAllowTranslucentSelection()
@@ -2237,17 +2242,17 @@ void FLevelEditorActionCallbacks::SetActorSnapSetting(float Distance)
 
 void FLevelEditorActionCallbacks::OnToggleHideViewportUI()
 {
-	GEditorModeTools().SetHideViewportUI( !GEditorModeTools().IsViewportUIHidden() );
+	GLevelEditorModeTools().SetHideViewportUI( !GLevelEditorModeTools().IsViewportUIHidden() );
 }
 
 bool FLevelEditorActionCallbacks::IsViewportUIHidden()
 {
-	return GEditorModeTools().IsViewportUIHidden();
+	return GLevelEditorModeTools().IsViewportUIHidden();
 }
 
 bool FLevelEditorActionCallbacks::IsEditorModeActive( FEditorModeID EditorMode )
 {
-	return GEditorModeTools().IsModeActive( EditorMode );
+	return GLevelEditorModeTools().IsModeActive( EditorMode );
 }
 
 void FLevelEditorActionCallbacks::OnAddVolume( UClass* VolumeClass )
@@ -2304,21 +2309,21 @@ void FLevelEditorActionCallbacks::SelectActorsInLayers()
 
 void FLevelEditorActionCallbacks::SetWidgetMode( FWidget::EWidgetMode WidgetMode )
 {
-	if( !GEditorModeTools().IsTracking() )
+	if( !GLevelEditorModeTools().IsTracking() )
 	{
-		GEditorModeTools().SetWidgetMode( WidgetMode );
+		GLevelEditorModeTools().SetWidgetMode( WidgetMode );
 		GEditor->RedrawAllViewports();
 	}
 }
 
 bool FLevelEditorActionCallbacks::IsWidgetModeActive( FWidget::EWidgetMode WidgetMode )
 {
-	return GEditorModeTools().GetWidgetMode() == WidgetMode;
+	return GLevelEditorModeTools().GetWidgetMode() == WidgetMode;
 }
 
 bool FLevelEditorActionCallbacks::CanSetWidgetMode( FWidget::EWidgetMode WidgetMode )
 {
-	return GEditorModeTools().GetShowWidget() == true;
+	return GLevelEditorModeTools().GetShowWidget() == true;
 }
 
 bool FLevelEditorActionCallbacks::IsTranslateRotateModeVisible()
@@ -2328,12 +2333,12 @@ bool FLevelEditorActionCallbacks::IsTranslateRotateModeVisible()
 
 void FLevelEditorActionCallbacks::SetCoordinateSystem( ECoordSystem CoordinateSystem )
 {
-	GEditorModeTools().SetCoordSystem( CoordinateSystem );
+	GLevelEditorModeTools().SetCoordSystem( CoordinateSystem );
 }
 
 bool FLevelEditorActionCallbacks::IsCoordinateSystemActive( ECoordSystem CoordinateSystem )
 {
-	return GEditorModeTools().GetCoordSystem() == CoordinateSystem;
+	return GLevelEditorModeTools().GetCoordSystem() == CoordinateSystem;
 }
 
 void FLevelEditorActionCallbacks::MoveActorToGrid_Clicked( bool InAlign, bool bInPerActor )
diff --git a/Engine/Source/Editor/LevelEditor/Private/LevelEditorContextMenu.cpp b/Engine/Source/Editor/LevelEditor/Private/LevelEditorContextMenu.cpp
--- a/Engine/Source/Editor/LevelEditor/Private/LevelEditorContextMenu.cpp
+++ b/Engine/Source/Editor/LevelEditor/Private/LevelEditorContextMenu.cpp
@@ -604,9 +604,9 @@ void FLevelEditorContextMenuImpl::FillMatineeSelectActorMenu( FMenuBuilder& Menu
 						MenuBuilder.AddMenuEntry( Text, Text, FSlateIcon(), CurMatineeActorAction );
 
 						// if matinee is opened, and if that is CurMatineeActor, show option to go to group
-						if( GEditorModeTools().IsModeActive( FBuiltinEditorModes::EM_InterpEdit ) )
+						if( GLevelEditorModeTools().IsModeActive( FBuiltinEditorModes::EM_InterpEdit ) )
 						{
-							const FEdModeInterpEdit* InterpEditMode = (const FEdModeInterpEdit*)GEditorModeTools().GetActiveMode( FBuiltinEditorModes::EM_InterpEdit );
+							const FEdModeInterpEdit* InterpEditMode = (const FEdModeInterpEdit*)GLevelEditorModeTools().GetActiveMode( FBuiltinEditorModes::EM_InterpEdit );
 
 							if ( InterpEditMode && InterpEditMode->MatineeActor == CurMatineeActor )
 							{
diff --git a/Engine/Source/Editor/LevelEditor/Private/LevelEditorModesActions.cpp b/Engine/Source/Editor/LevelEditor/Private/LevelEditorModesActions.cpp
--- a/Engine/Source/Editor/LevelEditor/Private/LevelEditorModesActions.cpp
+++ b/Engine/Source/Editor/LevelEditor/Private/LevelEditorModesActions.cpp
@@ -13,32 +13,18 @@ void FLevelEditorModesCommands::RegisterCommands()
 {
 	EditorModeCommands.Empty();
 
-	TArray<FEdMode*> Modes;
-	GEditorModeTools().GetModes(Modes);
-
 	int editorMode = 0;
-
-	struct FCompareEdModeByPriority
-	{
-		FORCEINLINE bool operator()(const FEdMode& A, const FEdMode& B) const
-		{
-			return A.GetPriorityOrder() < B.GetPriorityOrder();
-		}
-	};
-
-	Modes.Sort(FCompareEdModeByPriority());
-
 	FKey EdModeKeys[9] = { EKeys::One, EKeys::Two, EKeys::Three, EKeys::Four, EKeys::Five, EKeys::Six, EKeys::Seven, EKeys::Eight, EKeys::Nine };
 
-	for (FEdMode* Mode : Modes)
+	for ( const FEditorModeInfo& Mode : FEditorModeRegistry::Get().GetSortedModeInfo() )
 	{
 		// If the mode isn't visible don't create a menu option for it.
-		if (!Mode->IsVisible())
+		if (!Mode.bVisible)
 		{
 			continue;
 		}
 
-		FName EditorModeCommandName = FName(*(FString("EditorMode.") + Mode->GetID().ToString()));
+		FName EditorModeCommandName = FName(*(FString("EditorMode.") + Mode.ID.ToString()));
 
 		TSharedPtr<FUICommandInfo> EditorModeCommand = 
 			FInputBindingManager::Get().FindCommandInContext(GetContextName(), EditorModeCommandName);
@@ -47,16 +33,16 @@ void FLevelEditorModesCommands::RegisterCommands()
 		if ( !EditorModeCommand.IsValid() )
 		{
 			FFormatNamedArguments Args;
-			Args.Add( TEXT("Mode"), Mode->GetName() );
+			Args.Add( TEXT("Mode"), Mode.Name );
 			const FText Tooltip = FText::Format( NSLOCTEXT("LevelEditor", "ModeTooltipF", "Activate {Mode} Editing Mode"), Args );
 
 			FUICommandInfo::MakeCommandInfo(
 				this->AsShared(),
 				EditorModeCommand,
 				EditorModeCommandName,
-				Mode->GetName(),
+				Mode.Name,
 				Tooltip,
-				Mode->GetIcon(),
+				Mode.IconBrush,
 				EUserInterfaceActionType::ToggleButton,
 				editorMode < 9 ? FInputGesture( EModifierKey::Shift, EdModeKeys[editorMode] ) : FInputGesture() );
 
