Geometry edit mode now has a UI priority order

diff --git a/Engine/Source/Editor/GeometryMode/Private/GeometryEdMode.cpp b/Engine/Source/Editor/GeometryMode/Private/GeometryEdMode.cpp
--- a/Engine/Source/Editor/GeometryMode/Private/GeometryEdMode.cpp
+++ b/Engine/Source/Editor/GeometryMode/Private/GeometryEdMode.cpp
@@ -15,7 +15,7 @@ void FGeometryModeModule::StartupModule()
 		FBuiltinEditorModes::EM_Geometry,
 		NSLOCTEXT("EditorModes", "GeometryMode", "Geometry Editing"),
 		FSlateIcon(FEditorStyle::GetStyleSetName(), "LevelEditor.BspMode", "LevelEditor.BspMode.Small"),
-		true
+		true, 500
 		);
 }
 
