Rename FPostConstructInitializeProperties to something simpler

diff --git a/Engine/Source/Editor/BlueprintGraph/Private/K2Node_SwitchName.cpp b/Engine/Source/Editor/BlueprintGraph/Private/K2Node_SwitchName.cpp
--- a/Engine/Source/Editor/BlueprintGraph/Private/K2Node_SwitchName.cpp
+++ b/Engine/Source/Editor/BlueprintGraph/Private/K2Node_SwitchName.cpp
@@ -5,8 +5,8 @@
 #include "BlueprintNodeSpawner.h"
 #include "BlueprintActionDatabaseRegistrar.h"
 
-UK2Node_SwitchName::UK2Node_SwitchName(const class FPostConstructInitializeProperties& PCIP)
-	: Super(PCIP)
+UK2Node_SwitchName::UK2Node_SwitchName(const FObjectInitializer& ObjectInitializer)
+	: Super(ObjectInitializer)
 {
 	FunctionName = TEXT("NotEqual_NameName");
 	FunctionClass = UKismetMathLibrary::StaticClass();
diff --git a/Engine/Source/Editor/BlueprintGraph/Private/K2Node_SwitchString.cpp b/Engine/Source/Editor/BlueprintGraph/Private/K2Node_SwitchString.cpp
--- a/Engine/Source/Editor/BlueprintGraph/Private/K2Node_SwitchString.cpp
+++ b/Engine/Source/Editor/BlueprintGraph/Private/K2Node_SwitchString.cpp
@@ -7,8 +7,8 @@
 #include "BlueprintNodeSpawner.h"
 #include "BlueprintActionDatabaseRegistrar.h"
 
-UK2Node_SwitchString::UK2Node_SwitchString(const class FPostConstructInitializeProperties& PCIP)
-	: Super(PCIP)
+UK2Node_SwitchString::UK2Node_SwitchString(const FObjectInitializer& ObjectInitializer)
+	: Super(ObjectInitializer)
 {
 
 	bIsCaseSensitive = false;
diff --git a/Engine/Source/Editor/BlueprintGraph/Private/K2Node_TemporaryVariable.cpp b/Engine/Source/Editor/BlueprintGraph/Private/K2Node_TemporaryVariable.cpp
--- a/Engine/Source/Editor/BlueprintGraph/Private/K2Node_TemporaryVariable.cpp
+++ b/Engine/Source/Editor/BlueprintGraph/Private/K2Node_TemporaryVariable.cpp
@@ -31,8 +31,8 @@ public:
 	}
 };
 
-UK2Node_TemporaryVariable::UK2Node_TemporaryVariable(const class FPostConstructInitializeProperties& PCIP)
-	: Super(PCIP)
+UK2Node_TemporaryVariable::UK2Node_TemporaryVariable(const FObjectInitializer& ObjectInitializer)
+	: Super(ObjectInitializer)
 	, bIsPersistent(false)
 {
 }
diff --git a/Engine/Source/Editor/BlueprintGraph/Private/K2Node_Timeline.cpp b/Engine/Source/Editor/BlueprintGraph/Private/K2Node_Timeline.cpp
--- a/Engine/Source/Editor/BlueprintGraph/Private/K2Node_Timeline.cpp
+++ b/Engine/Source/Editor/BlueprintGraph/Private/K2Node_Timeline.cpp
@@ -15,8 +15,8 @@
 /////////////////////////////////////////////////////
 // UK2Node_Timeline
 
-UK2Node_Timeline::UK2Node_Timeline(const class FPostConstructInitializeProperties& PCIP)
-	: Super(PCIP)
+UK2Node_Timeline::UK2Node_Timeline(const FObjectInitializer& ObjectInitializer)
+	: Super(ObjectInitializer)
 {
 	bAutoPlay = false;
 	bLoop = false;
diff --git a/Engine/Source/Editor/BlueprintGraph/Private/K2Node_Tunnel.cpp b/Engine/Source/Editor/BlueprintGraph/Private/K2Node_Tunnel.cpp
--- a/Engine/Source/Editor/BlueprintGraph/Private/K2Node_Tunnel.cpp
+++ b/Engine/Source/Editor/BlueprintGraph/Private/K2Node_Tunnel.cpp
@@ -6,8 +6,8 @@
 #include "K2Node_MacroInstance.h"
 #include "K2Node_Tunnel.h"
 
-UK2Node_Tunnel::UK2Node_Tunnel(const class FPostConstructInitializeProperties& PCIP)
-	: Super(PCIP)
+UK2Node_Tunnel::UK2Node_Tunnel(const FObjectInitializer& ObjectInitializer)
+	: Super(ObjectInitializer)
 {
 	bIsEditable = true;
 }
diff --git a/Engine/Source/Editor/BlueprintGraph/Private/K2Node_Variable.cpp b/Engine/Source/Editor/BlueprintGraph/Private/K2Node_Variable.cpp
--- a/Engine/Source/Editor/BlueprintGraph/Private/K2Node_Variable.cpp
+++ b/Engine/Source/Editor/BlueprintGraph/Private/K2Node_Variable.cpp
@@ -8,8 +8,8 @@
 
 #define LOCTEXT_NAMESPACE "K2Node"
 
-UK2Node_Variable::UK2Node_Variable(const class FPostConstructInitializeProperties& PCIP)
-	: Super(PCIP)
+UK2Node_Variable::UK2Node_Variable(const FObjectInitializer& ObjectInitializer)
+	: Super(ObjectInitializer)
 {
 }
 
diff --git a/Engine/Source/Editor/BlueprintGraph/Private/K2Node_VariableGet.cpp b/Engine/Source/Editor/BlueprintGraph/Private/K2Node_VariableGet.cpp
--- a/Engine/Source/Editor/BlueprintGraph/Private/K2Node_VariableGet.cpp
+++ b/Engine/Source/Editor/BlueprintGraph/Private/K2Node_VariableGet.cpp
@@ -72,8 +72,8 @@ static FText K2Node_VariableGetImpl::GetBaseTooltip(FName VarName)
 
 }
 
-UK2Node_VariableGet::UK2Node_VariableGet(const class FPostConstructInitializeProperties& PCIP)
-	: Super(PCIP)
+UK2Node_VariableGet::UK2Node_VariableGet(const FObjectInitializer& ObjectInitializer)
+	: Super(ObjectInitializer)
 {
 }
 
diff --git a/Engine/Source/Editor/BlueprintGraph/Private/K2Node_VariableSet.cpp b/Engine/Source/Editor/BlueprintGraph/Private/K2Node_VariableSet.cpp
--- a/Engine/Source/Editor/BlueprintGraph/Private/K2Node_VariableSet.cpp
+++ b/Engine/Source/Editor/BlueprintGraph/Private/K2Node_VariableSet.cpp
@@ -68,8 +68,8 @@ static bool K2Node_VariableSetImpl::PropertyHasLocalRepNotify(UProperty const* V
 	return false;
 }
 
-UK2Node_VariableSet::UK2Node_VariableSet(const class FPostConstructInitializeProperties& PCIP)
-	: Super(PCIP)
+UK2Node_VariableSet::UK2Node_VariableSet(const FObjectInitializer& ObjectInitializer)
+	: Super(ObjectInitializer)
 {
 }
 
diff --git a/Engine/Source/Editor/BlueprintGraph/Private/K2Node_VariableSetRef.cpp b/Engine/Source/Editor/BlueprintGraph/Private/K2Node_VariableSetRef.cpp
--- a/Engine/Source/Editor/BlueprintGraph/Private/K2Node_VariableSetRef.cpp
+++ b/Engine/Source/Editor/BlueprintGraph/Private/K2Node_VariableSetRef.cpp
@@ -83,8 +83,8 @@ public:
 	}
 };
 
-UK2Node_VariableSetRef::UK2Node_VariableSetRef(const class FPostConstructInitializeProperties& PCIP)
-	: Super(PCIP)
+UK2Node_VariableSetRef::UK2Node_VariableSetRef(const FObjectInitializer& ObjectInitializer)
+	: Super(ObjectInitializer)
 {
 }
 
diff --git a/Engine/Source/Editor/BlueprintGraph/Private/NodeDependingOnEnumInterface.cpp b/Engine/Source/Editor/BlueprintGraph/Private/NodeDependingOnEnumInterface.cpp
--- a/Engine/Source/Editor/BlueprintGraph/Private/NodeDependingOnEnumInterface.cpp
+++ b/Engine/Source/Editor/BlueprintGraph/Private/NodeDependingOnEnumInterface.cpp
@@ -4,7 +4,7 @@
 
 #include "NodeDependingOnEnumInterface.h"
 
-UNodeDependingOnEnumInterface::UNodeDependingOnEnumInterface(const class FPostConstructInitializeProperties& PCIP)
-	: Super(PCIP)
+UNodeDependingOnEnumInterface::UNodeDependingOnEnumInterface(const FObjectInitializer& ObjectInitializer)
+	: Super(ObjectInitializer)
 {
 }
\ No newline at end of file
diff --git a/Engine/Source/Editor/Blutility/Private/EditorUtilityBlueprint.cpp b/Engine/Source/Editor/Blutility/Private/EditorUtilityBlueprint.cpp
--- a/Engine/Source/Editor/Blutility/Private/EditorUtilityBlueprint.cpp
+++ b/Engine/Source/Editor/Blutility/Private/EditorUtilityBlueprint.cpp
@@ -6,7 +6,7 @@
 /////////////////////////////////////////////////////
 // UEditorUtilityBlueprint
 
-UEditorUtilityBlueprint::UEditorUtilityBlueprint(const FPostConstructInitializeProperties& PCIP)
-	: Super(PCIP)
+UEditorUtilityBlueprint::UEditorUtilityBlueprint(const FObjectInitializer& ObjectInitializer)
+	: Super(ObjectInitializer)
 {
 }
diff --git a/Engine/Source/Editor/Blutility/Private/EditorUtilityBlueprintFactory.cpp b/Engine/Source/Editor/Blutility/Private/EditorUtilityBlueprintFactory.cpp
--- a/Engine/Source/Editor/Blutility/Private/EditorUtilityBlueprintFactory.cpp
+++ b/Engine/Source/Editor/Blutility/Private/EditorUtilityBlueprintFactory.cpp
@@ -26,8 +26,8 @@ public:
 /////////////////////////////////////////////////////
 // UEditorUtilityBlueprintFactory
 
-UEditorUtilityBlueprintFactory::UEditorUtilityBlueprintFactory(const class FPostConstructInitializeProperties& PCIP)
-	: Super(PCIP)
+UEditorUtilityBlueprintFactory::UEditorUtilityBlueprintFactory(const FObjectInitializer& ObjectInitializer)
+	: Super(ObjectInitializer)
 {
 	UClass* DefaultParentClass = APlacedEditorUtilityBase::StaticClass();
 
diff --git a/Engine/Source/Editor/Blutility/Private/GlobalEditorUtilityBase.cpp b/Engine/Source/Editor/Blutility/Private/GlobalEditorUtilityBase.cpp
--- a/Engine/Source/Editor/Blutility/Private/GlobalEditorUtilityBase.cpp
+++ b/Engine/Source/Editor/Blutility/Private/GlobalEditorUtilityBase.cpp
@@ -9,8 +9,8 @@
 
 /////////////////////////////////////////////////////
 
-UGlobalEditorUtilityBase::UGlobalEditorUtilityBase(const FPostConstructInitializeProperties& PCIP)
-	: Super(PCIP)
+UGlobalEditorUtilityBase::UGlobalEditorUtilityBase(const FObjectInitializer& ObjectInitializer)
+	: Super(ObjectInitializer)
 {
 }
 
diff --git a/Engine/Source/Editor/Blutility/Private/PlacedEditorUtilityBase.cpp b/Engine/Source/Editor/Blutility/Private/PlacedEditorUtilityBase.cpp
--- a/Engine/Source/Editor/Blutility/Private/PlacedEditorUtilityBase.cpp
+++ b/Engine/Source/Editor/Blutility/Private/PlacedEditorUtilityBase.cpp
@@ -6,8 +6,8 @@
 /////////////////////////////////////////////////////
 // APlacedEditorUtilityBase
 
-APlacedEditorUtilityBase::APlacedEditorUtilityBase(const FPostConstructInitializeProperties& PCIP)
-	: Super(PCIP)
+APlacedEditorUtilityBase::APlacedEditorUtilityBase(const FObjectInitializer& ObjectInitializer)
+	: Super(ObjectInitializer)
 {
 	PrimaryActorTick.bCanEverTick = true;
 	HelpText = TEXT("Please fill out the help text");
diff --git a/Engine/Source/Editor/Cascade/Private/Cascade.cpp b/Engine/Source/Editor/Cascade/Private/Cascade.cpp
--- a/Engine/Source/Editor/Cascade/Private/Cascade.cpp
+++ b/Engine/Source/Editor/Cascade/Private/Cascade.cpp
@@ -5045,8 +5045,8 @@ void FCascade::CloseEntryPopup()
 /*-----------------------------------------------------------------------------
 	UCascadeParticleSystemComponent
 -----------------------------------------------------------------------------*/
-UCascadeParticleSystemComponent::UCascadeParticleSystemComponent(const class FPostConstructInitializeProperties& PCIP)
-	: Super(PCIP)
+UCascadeParticleSystemComponent::UCascadeParticleSystemComponent(const FObjectInitializer& ObjectInitializer)
+	: Super(ObjectInitializer)
 {
 }
 
@@ -5090,8 +5090,8 @@ void UCascadeParticleSystemComponent::UpdateLODInformation()
 	}
 }
 
-UCascadeConfiguration::UCascadeConfiguration(const class FPostConstructInitializeProperties& PCIP)
-	: Super(PCIP)
+UCascadeConfiguration::UCascadeConfiguration(const FObjectInitializer& ObjectInitializer)
+	: Super(ObjectInitializer)
 {
 }
 
diff --git a/Engine/Source/Editor/DestructibleMeshEditor/Private/DestructibleChunkParamsProxy.cpp b/Engine/Source/Editor/DestructibleMeshEditor/Private/DestructibleChunkParamsProxy.cpp
--- a/Engine/Source/Editor/DestructibleMeshEditor/Private/DestructibleChunkParamsProxy.cpp
+++ b/Engine/Source/Editor/DestructibleMeshEditor/Private/DestructibleChunkParamsProxy.cpp
@@ -4,8 +4,8 @@
 #include "ApexDestructibleAssetImport.h"
 #include "Engine/DestructibleMesh.h"
 
-UDestructibleChunkParamsProxy::UDestructibleChunkParamsProxy(const class FPostConstructInitializeProperties& PCIP)
-	:Super(PCIP)
+UDestructibleChunkParamsProxy::UDestructibleChunkParamsProxy(const FObjectInitializer& ObjectInitializer)
+	:Super(ObjectInitializer)
 {
 
 }
diff --git a/Engine/Source/Editor/EditorLiveStreaming/EditorLiveStreamingSettings.cpp b/Engine/Source/Editor/EditorLiveStreaming/EditorLiveStreamingSettings.cpp
--- a/Engine/Source/Editor/EditorLiveStreaming/EditorLiveStreamingSettings.cpp
+++ b/Engine/Source/Editor/EditorLiveStreaming/EditorLiveStreamingSettings.cpp
@@ -3,8 +3,8 @@
 #include "EditorLiveStreamingSettings.h"
 
 
-UEditorLiveStreamingSettings::UEditorLiveStreamingSettings( const class FPostConstructInitializeProperties& PCIP )
-	: Super( PCIP )
+UEditorLiveStreamingSettings::UEditorLiveStreamingSettings( const FObjectInitializer& ObjectInitializer )
+	: Super( ObjectInitializer )
 {
 }
 
diff --git a/Engine/Source/Editor/EditorStyle/Private/EditorStyleClasses.cpp b/Engine/Source/Editor/EditorStyle/Private/EditorStyleClasses.cpp
--- a/Engine/Source/Editor/EditorStyle/Private/EditorStyleClasses.cpp
+++ b/Engine/Source/Editor/EditorStyle/Private/EditorStyleClasses.cpp
@@ -6,8 +6,8 @@
 /* UEditorStyleSettings interface
  *****************************************************************************/
 
-UEditorStyleSettings::UEditorStyleSettings( const class FPostConstructInitializeProperties& PCIP )
-	: Super(PCIP)
+UEditorStyleSettings::UEditorStyleSettings( const FObjectInitializer& ObjectInitializer )
+	: Super(ObjectInitializer)
 {
 	SelectionColor = FLinearColor(0.828f, 0.364f, 0.003f);
 	InactiveSelectionColor = FLinearColor(0.25f, 0.25f, 0.25f);
diff --git a/Engine/Source/Editor/EnvironmentQueryEditor/Private/EdGraphSchema_EnvironmentQuery.cpp b/Engine/Source/Editor/EnvironmentQueryEditor/Private/EdGraphSchema_EnvironmentQuery.cpp
--- a/Engine/Source/Editor/EnvironmentQueryEditor/Private/EdGraphSchema_EnvironmentQuery.cpp
+++ b/Engine/Source/Editor/EnvironmentQueryEditor/Private/EdGraphSchema_EnvironmentQuery.cpp
@@ -117,8 +117,8 @@ void FEnvironmentQuerySchemaAction_NewSubNode::AddReferencedObjects( FReferenceC
 
 //////////////////////////////////////////////////////////////////////////
 
-UEdGraphSchema_EnvironmentQuery::UEdGraphSchema_EnvironmentQuery(const class FPostConstructInitializeProperties& PCIP)
-	: Super(PCIP)
+UEdGraphSchema_EnvironmentQuery::UEdGraphSchema_EnvironmentQuery(const FObjectInitializer& ObjectInitializer)
+	: Super(ObjectInitializer)
 {
 }
 
diff --git a/Engine/Source/Editor/EnvironmentQueryEditor/Private/EnvironmentQueryFactory.cpp b/Engine/Source/Editor/EnvironmentQueryEditor/Private/EnvironmentQueryFactory.cpp
--- a/Engine/Source/Editor/EnvironmentQueryEditor/Private/EnvironmentQueryFactory.cpp
+++ b/Engine/Source/Editor/EnvironmentQueryEditor/Private/EnvironmentQueryFactory.cpp
@@ -4,8 +4,8 @@
 #include "EnvironmentQueryEditorModule.h"
 #include "EnvironmentQuery/EnvQuery.h"
 
-UEnvironmentQueryFactory::UEnvironmentQueryFactory(const class FPostConstructInitializeProperties& PCIP)
-	: Super(PCIP)
+UEnvironmentQueryFactory::UEnvironmentQueryFactory(const FObjectInitializer& ObjectInitializer)
+	: Super(ObjectInitializer)
 {
 	SupportedClass = UEnvQuery::StaticClass();
 	bEditAfterNew = true;
diff --git a/Engine/Source/Editor/EnvironmentQueryEditor/Private/EnvironmentQueryGraph.cpp b/Engine/Source/Editor/EnvironmentQueryEditor/Private/EnvironmentQueryGraph.cpp
--- a/Engine/Source/Editor/EnvironmentQueryEditor/Private/EnvironmentQueryGraph.cpp
+++ b/Engine/Source/Editor/EnvironmentQueryEditor/Private/EnvironmentQueryGraph.cpp
@@ -15,7 +15,7 @@ namespace EQSGraphVersion
 	const int32 Latest = CopyPasteOutersBug;
 }
 
-UEnvironmentQueryGraph::UEnvironmentQueryGraph(const class FPostConstructInitializeProperties& PCIP) : Super(PCIP)
+UEnvironmentQueryGraph::UEnvironmentQueryGraph(const FObjectInitializer& ObjectInitializer) : Super(ObjectInitializer)
 {
 	Schema = UEdGraphSchema_EnvironmentQuery::StaticClass();
 	bLockUpdates = false;
diff --git a/Engine/Source/Editor/EnvironmentQueryEditor/Private/EnvironmentQueryGraphNode.cpp b/Engine/Source/Editor/EnvironmentQueryEditor/Private/EnvironmentQueryGraphNode.cpp
--- a/Engine/Source/Editor/EnvironmentQueryEditor/Private/EnvironmentQueryGraphNode.cpp
+++ b/Engine/Source/Editor/EnvironmentQueryEditor/Private/EnvironmentQueryGraphNode.cpp
@@ -3,8 +3,8 @@
 #include "EnvironmentQueryEditorPrivatePCH.h"
 #include "ScopedTransaction.h"
 
-UEnvironmentQueryGraphNode::UEnvironmentQueryGraphNode(const class FPostConstructInitializeProperties& PCIP)
-	: Super(PCIP)
+UEnvironmentQueryGraphNode::UEnvironmentQueryGraphNode(const FObjectInitializer& ObjectInitializer)
+	: Super(ObjectInitializer)
 	, NodeInstance(NULL)
 	, EnvQueryNodeClass(NULL)
 {
diff --git a/Engine/Source/Editor/EnvironmentQueryEditor/Private/EnvironmentQueryGraphNode_Option.cpp b/Engine/Source/Editor/EnvironmentQueryEditor/Private/EnvironmentQueryGraphNode_Option.cpp
--- a/Engine/Source/Editor/EnvironmentQueryEditor/Private/EnvironmentQueryGraphNode_Option.cpp
+++ b/Engine/Source/Editor/EnvironmentQueryEditor/Private/EnvironmentQueryGraphNode_Option.cpp
@@ -8,8 +8,8 @@
 
 #define LOCTEXT_NAMESPACE "BehaviorTreeGraphNode"
 
-UEnvironmentQueryGraphNode_Option::UEnvironmentQueryGraphNode_Option(const class FPostConstructInitializeProperties& PCIP)
-	: Super(PCIP)
+UEnvironmentQueryGraphNode_Option::UEnvironmentQueryGraphNode_Option(const FObjectInitializer& ObjectInitializer)
+	: Super(ObjectInitializer)
 {
 }
 
diff --git a/Engine/Source/Editor/EnvironmentQueryEditor/Private/EnvironmentQueryGraphNode_Root.cpp b/Engine/Source/Editor/EnvironmentQueryEditor/Private/EnvironmentQueryGraphNode_Root.cpp
--- a/Engine/Source/Editor/EnvironmentQueryEditor/Private/EnvironmentQueryGraphNode_Root.cpp
+++ b/Engine/Source/Editor/EnvironmentQueryEditor/Private/EnvironmentQueryGraphNode_Root.cpp
@@ -2,7 +2,7 @@
 
 #include "EnvironmentQueryEditorPrivatePCH.h"
 
-UEnvironmentQueryGraphNode_Root::UEnvironmentQueryGraphNode_Root(const class FPostConstructInitializeProperties& PCIP) : Super(PCIP)
+UEnvironmentQueryGraphNode_Root::UEnvironmentQueryGraphNode_Root(const FObjectInitializer& ObjectInitializer) : Super(ObjectInitializer)
 {
 }
 
diff --git a/Engine/Source/Editor/EnvironmentQueryEditor/Private/EnvironmentQueryGraphNode_Test.cpp b/Engine/Source/Editor/EnvironmentQueryEditor/Private/EnvironmentQueryGraphNode_Test.cpp
--- a/Engine/Source/Editor/EnvironmentQueryEditor/Private/EnvironmentQueryGraphNode_Test.cpp
+++ b/Engine/Source/Editor/EnvironmentQueryEditor/Private/EnvironmentQueryGraphNode_Test.cpp
@@ -3,7 +3,7 @@
 #include "EnvironmentQueryEditorPrivatePCH.h"
 #include "EnvironmentQuery/EnvQueryTest.h"
 
-UEnvironmentQueryGraphNode_Test::UEnvironmentQueryGraphNode_Test(const class FPostConstructInitializeProperties& PCIP) : Super(PCIP)
+UEnvironmentQueryGraphNode_Test::UEnvironmentQueryGraphNode_Test(const FObjectInitializer& ObjectInitializer) : Super(ObjectInitializer)
 {
 	bTestEnabled = true;
 	bHasNamedWeight = false;
diff --git a/Engine/Source/Editor/GameProjectGeneration/Private/DefaultTemplateProjectDefs.cpp b/Engine/Source/Editor/GameProjectGeneration/Private/DefaultTemplateProjectDefs.cpp
--- a/Engine/Source/Editor/GameProjectGeneration/Private/DefaultTemplateProjectDefs.cpp
+++ b/Engine/Source/Editor/GameProjectGeneration/Private/DefaultTemplateProjectDefs.cpp
@@ -2,8 +2,8 @@
 
 #include "GameProjectGenerationPrivatePCH.h"
 
-UDefaultTemplateProjectDefs::UDefaultTemplateProjectDefs(const FPostConstructInitializeProperties& PCIP)
-	: Super(PCIP)
+UDefaultTemplateProjectDefs::UDefaultTemplateProjectDefs(const FObjectInitializer& ObjectInitializer)
+	: Super(ObjectInitializer)
 {
 }
 
diff --git a/Engine/Source/Editor/GameProjectGeneration/Private/TemplateProjectDefs.cpp b/Engine/Source/Editor/GameProjectGeneration/Private/TemplateProjectDefs.cpp
--- a/Engine/Source/Editor/GameProjectGeneration/Private/TemplateProjectDefs.cpp
+++ b/Engine/Source/Editor/GameProjectGeneration/Private/TemplateProjectDefs.cpp
@@ -54,8 +54,8 @@ FTemplateConfigValue::FTemplateConfigValue(const FString& InFile, const FString&
 
 
 
-UTemplateProjectDefs::UTemplateProjectDefs(const class FPostConstructInitializeProperties& PCIP)
-	: Super(PCIP)
+UTemplateProjectDefs::UTemplateProjectDefs(const FObjectInitializer& ObjectInitializer)
+	: Super(ObjectInitializer)
 {
 
 }
diff --git a/Engine/Source/Editor/GameplayAbilitiesEditor/Private/GameplayAbilitiesBlueprintFactory.cpp b/Engine/Source/Editor/GameplayAbilitiesEditor/Private/GameplayAbilitiesBlueprintFactory.cpp
--- a/Engine/Source/Editor/GameplayAbilitiesEditor/Private/GameplayAbilitiesBlueprintFactory.cpp
+++ b/Engine/Source/Editor/GameplayAbilitiesEditor/Private/GameplayAbilitiesBlueprintFactory.cpp
@@ -234,8 +234,8 @@ private:
 	UGameplayAbilitiesBlueprintFactory implementation.
 ------------------------------------------------------------------------------*/
 
-UGameplayAbilitiesBlueprintFactory::UGameplayAbilitiesBlueprintFactory(const class FPostConstructInitializeProperties& PCIP)
-	: Super(PCIP)
+UGameplayAbilitiesBlueprintFactory::UGameplayAbilitiesBlueprintFactory(const FObjectInitializer& ObjectInitializer)
+	: Super(ObjectInitializer)
 {
 	bCreateNew = true;
 	bEditAfterNew = true;
diff --git a/Engine/Source/Editor/GameplayAbilitiesEditor/Private/GameplayAbilityGraph.cpp b/Engine/Source/Editor/GameplayAbilitiesEditor/Private/GameplayAbilityGraph.cpp
--- a/Engine/Source/Editor/GameplayAbilitiesEditor/Private/GameplayAbilityGraph.cpp
+++ b/Engine/Source/Editor/GameplayAbilitiesEditor/Private/GameplayAbilityGraph.cpp
@@ -8,8 +8,8 @@
 /////////////////////////////////////////////////////
 // UGameplayAbilityGraph
 
-UGameplayAbilityGraph::UGameplayAbilityGraph(const FPostConstructInitializeProperties& PCIP)
-	: Super(PCIP)
+UGameplayAbilityGraph::UGameplayAbilityGraph(const FObjectInitializer& ObjectInitializer)
+	: Super(ObjectInitializer)
 {
 }
 
diff --git a/Engine/Source/Editor/GameplayAbilitiesEditor/Private/GameplayAbilityGraphSchema.cpp b/Engine/Source/Editor/GameplayAbilitiesEditor/Private/GameplayAbilityGraphSchema.cpp
--- a/Engine/Source/Editor/GameplayAbilitiesEditor/Private/GameplayAbilityGraphSchema.cpp
+++ b/Engine/Source/Editor/GameplayAbilitiesEditor/Private/GameplayAbilityGraphSchema.cpp
@@ -8,8 +8,8 @@
 #include "BlueprintEditorUtils.h"
 
 
-UGameplayAbilityGraphSchema::UGameplayAbilityGraphSchema(const FPostConstructInitializeProperties& PCIP)
-: Super(PCIP)
+UGameplayAbilityGraphSchema::UGameplayAbilityGraphSchema(const FObjectInitializer& ObjectInitializer)
+: Super(ObjectInitializer)
 {
 }
 
diff --git a/Engine/Source/Editor/GameplayAbilitiesEditor/Private/K2Node_GameplayCueEvent.cpp b/Engine/Source/Editor/GameplayAbilitiesEditor/Private/K2Node_GameplayCueEvent.cpp
--- a/Engine/Source/Editor/GameplayAbilitiesEditor/Private/K2Node_GameplayCueEvent.cpp
+++ b/Engine/Source/Editor/GameplayAbilitiesEditor/Private/K2Node_GameplayCueEvent.cpp
@@ -13,8 +13,8 @@
 
 #define LOCTEXT_NAMESPACE "K2Node_GameplayCueEvent"
 
-UK2Node_GameplayCueEvent::UK2Node_GameplayCueEvent(const class FPostConstructInitializeProperties& PCIP)
-	: Super(PCIP)
+UK2Node_GameplayCueEvent::UK2Node_GameplayCueEvent(const FObjectInitializer& ObjectInitializer)
+	: Super(ObjectInitializer)
 {
 	EventSignatureName = GAMEPLAYABILITIES_BlueprintCustomHandler;
 	EventSignatureClass = UGameplayCueInterface::StaticClass();
diff --git a/Engine/Source/Editor/GameplayAbilitiesEditor/Private/K2Node_GameplayEffectVariable.cpp b/Engine/Source/Editor/GameplayAbilitiesEditor/Private/K2Node_GameplayEffectVariable.cpp
--- a/Engine/Source/Editor/GameplayAbilitiesEditor/Private/K2Node_GameplayEffectVariable.cpp
+++ b/Engine/Source/Editor/GameplayAbilitiesEditor/Private/K2Node_GameplayEffectVariable.cpp
@@ -6,8 +6,8 @@
 
 #define LOCTEXT_NAMESPACE "K2Node_GameplayEffectVariable"
 
-UK2Node_GameplayEffectVariable::UK2Node_GameplayEffectVariable(const class FPostConstructInitializeProperties& PCIP)
-	: Super(PCIP)
+UK2Node_GameplayEffectVariable::UK2Node_GameplayEffectVariable(const FObjectInitializer& ObjectInitializer)
+	: Super(ObjectInitializer)
 {
 }
 
diff --git a/Engine/Source/Editor/GameplayAbilitiesEditor/Private/K2Node_LatentAbilityCall.cpp b/Engine/Source/Editor/GameplayAbilitiesEditor/Private/K2Node_LatentAbilityCall.cpp
--- a/Engine/Source/Editor/GameplayAbilitiesEditor/Private/K2Node_LatentAbilityCall.cpp
+++ b/Engine/Source/Editor/GameplayAbilitiesEditor/Private/K2Node_LatentAbilityCall.cpp
@@ -15,8 +15,8 @@
 
 #define LOCTEXT_NAMESPACE "K2Node"
 
-UK2Node_LatentAbilityCall::UK2Node_LatentAbilityCall(const class FPostConstructInitializeProperties& PCIP)
-	: Super(PCIP)
+UK2Node_LatentAbilityCall::UK2Node_LatentAbilityCall(const FObjectInitializer& ObjectInitializer)
+	: Super(ObjectInitializer)
 {
 	ProxyActivateFunctionName = GET_FUNCTION_NAME_CHECKED(UAbilityTask, Activate);
 }
diff --git a/Engine/Source/Editor/GameplayTagsEditor/Private/GameplayTagsK2Node_LiteralGameplayTag.cpp b/Engine/Source/Editor/GameplayTagsEditor/Private/GameplayTagsK2Node_LiteralGameplayTag.cpp
--- a/Engine/Source/Editor/GameplayTagsEditor/Private/GameplayTagsK2Node_LiteralGameplayTag.cpp
+++ b/Engine/Source/Editor/GameplayTagsEditor/Private/GameplayTagsK2Node_LiteralGameplayTag.cpp
@@ -10,8 +10,8 @@
 
 #define LOCTEXT_NAMESPACE "GameplayTagsK2Node_LiteralGameplayTag"
 
-UGameplayTagsK2Node_LiteralGameplayTag::UGameplayTagsK2Node_LiteralGameplayTag(const class FPostConstructInitializeProperties& PCIP)
-	: Super(PCIP)
+UGameplayTagsK2Node_LiteralGameplayTag::UGameplayTagsK2Node_LiteralGameplayTag(const FObjectInitializer& ObjectInitializer)
+	: Super(ObjectInitializer)
 {
 
 }
diff --git a/Engine/Source/Editor/GeometryMode/Private/GeometryModifiers.cpp b/Engine/Source/Editor/GeometryMode/Private/GeometryModifiers.cpp
--- a/Engine/Source/Editor/GeometryMode/Private/GeometryModifiers.cpp
+++ b/Engine/Source/Editor/GeometryMode/Private/GeometryModifiers.cpp
@@ -26,8 +26,8 @@ static FVector ComputeWorldSpaceMousePos( FEditorViewportClient* ViewportClient
 	return View->PixelToWorld(ViewportClient->Viewport->GetMouseX(),ViewportClient->Viewport->GetMouseY(),0.5f);
 }
 
-UGeomModifier::UGeomModifier(const class FPostConstructInitializeProperties& PCIP)
-	: Super(PCIP)
+UGeomModifier::UGeomModifier(const FObjectInitializer& ObjectInitializer)
+	: Super(ObjectInitializer)
 {
 	bPushButton = false;
 	bInitialized = false;
@@ -366,8 +366,8 @@ void UGeomModifier::StoreAllCurrentGeomSelections()
 }
 
 
-UGeomModifier_Edit::UGeomModifier_Edit(const class FPostConstructInitializeProperties& PCIP)
-	: Super(PCIP)
+UGeomModifier_Edit::UGeomModifier_Edit(const FObjectInitializer& ObjectInitializer)
+	: Super(ObjectInitializer)
 {
 	Description = NSLOCTEXT("UnrealEd", "Edit", "Edit");
 }
@@ -606,8 +606,8 @@ bool UGeomModifier_Edit::InputDelta(FEditorViewportClient* InViewportClient,FVie
 /*------------------------------------------------------------------------------
 	UGeomModifier_Extrude
 ------------------------------------------------------------------------------*/
-UGeomModifier_Extrude::UGeomModifier_Extrude(const class FPostConstructInitializeProperties& PCIP)
-	: Super(PCIP)
+UGeomModifier_Extrude::UGeomModifier_Extrude(const FObjectInitializer& ObjectInitializer)
+	: Super(ObjectInitializer)
 {
 	Description = NSLOCTEXT("UnrealEd", "Extrude", "Extrude");
 	Length = 16;
@@ -837,8 +837,8 @@ void UGeomModifier_Extrude::Apply(int32 InLength, int32 InSegments)
 	}
 }
 
-UGeomModifier_Lathe::UGeomModifier_Lathe(const class FPostConstructInitializeProperties& PCIP)
-	: Super(PCIP)
+UGeomModifier_Lathe::UGeomModifier_Lathe(const FObjectInitializer& ObjectInitializer)
+	: Super(ObjectInitializer)
 {
 	Description = NSLOCTEXT("UnrealEd", "Lathe", "Lathe");
 	Axis = EAxis::Y;
@@ -1166,8 +1166,8 @@ void UGeomModifier_Lathe::Apply( int32 InTotalSegments, int32 InSegments, EAxis:
 /*------------------------------------------------------------------------------
 	UGeomModifier_Pen
 ------------------------------------------------------------------------------*/
-UGeomModifier_Pen::UGeomModifier_Pen(const class FPostConstructInitializeProperties& PCIP)
-	: Super(PCIP)
+UGeomModifier_Pen::UGeomModifier_Pen(const FObjectInitializer& ObjectInitializer)
+	: Super(ObjectInitializer)
 {
 	Description = NSLOCTEXT("UnrealEd", "Pen", "Pen");
 	bCreateBrushShape = false;
@@ -1699,7 +1699,7 @@ static ABrush* ClipBrushAgainstPlane( const FPlane& InPlane, ABrush* InBrush)
 	// perhaps there were additional brushes were selected. 
 	check( ClippedBrush->GetClass() == InBrush->GetClass() );
 
-	ClippedBrush->Brush = new( InBrush->GetOuter(), NAME_None )UModel( FPostConstructInitializeProperties(),NULL );
+	ClippedBrush->Brush = new( InBrush->GetOuter(), NAME_None )UModel( FObjectInitializer(),NULL );
 	ClippedBrush->BrushComponent->Brush = ClippedBrush->Brush;
 
 	GeometryClipping::BuildGiantAlignedBrush( *ClippedBrush, InPlane );
@@ -1816,8 +1816,8 @@ static ABrush* ClipBrushAgainstPlane( const FPlane& InPlane, ABrush* InBrush)
 }
 
 } // namespace GeometryClipping
-UGeomModifier_Clip::UGeomModifier_Clip(const class FPostConstructInitializeProperties& PCIP)
-	: Super(PCIP)
+UGeomModifier_Clip::UGeomModifier_Clip(const FObjectInitializer& ObjectInitializer)
+	: Super(ObjectInitializer)
 {
 	Description = NSLOCTEXT("UnrealEd", "BrushClip", "Brush Clip");
 	bFlipNormal = false;
@@ -2196,8 +2196,8 @@ void UGeomModifier_Clip::Tick(FEditorViewportClient* ViewportClient,float DeltaT
 }
 
 
-UGeomModifier_Delete::UGeomModifier_Delete(const class FPostConstructInitializeProperties& PCIP)
-	: Super(PCIP)
+UGeomModifier_Delete::UGeomModifier_Delete(const FObjectInitializer& ObjectInitializer)
+	: Super(ObjectInitializer)
 {
 	Description = NSLOCTEXT("UnrealEd", "Delete", "Delete");
 	bPushButton = true;
@@ -2281,8 +2281,8 @@ bool UGeomModifier_Delete::OnApply()
 	return bHandled;
 }
 
-UGeomModifier_Create::UGeomModifier_Create(const class FPostConstructInitializeProperties& PCIP)
-	: Super(PCIP)
+UGeomModifier_Create::UGeomModifier_Create(const FObjectInitializer& ObjectInitializer)
+	: Super(ObjectInitializer)
 {
 	Description = NSLOCTEXT("UnrealEd", "Create", "Create");
 	bPushButton = true;
@@ -2345,8 +2345,8 @@ bool UGeomModifier_Create::OnApply()
 	return true;
 }
 
-UGeomModifier_Flip::UGeomModifier_Flip(const class FPostConstructInitializeProperties& PCIP)
-	: Super(PCIP)
+UGeomModifier_Flip::UGeomModifier_Flip(const FObjectInitializer& ObjectInitializer)
+	: Super(ObjectInitializer)
 {
 	Description = NSLOCTEXT("UnrealEd", "Flip", "Flip");
 	bPushButton = true;
@@ -2389,8 +2389,8 @@ bool UGeomModifier_Flip::OnApply()
 	return true;
 }
 
-UGeomModifier_Split::UGeomModifier_Split(const class FPostConstructInitializeProperties& PCIP)
-	: Super(PCIP)
+UGeomModifier_Split::UGeomModifier_Split(const FObjectInitializer& ObjectInitializer)
+	: Super(ObjectInitializer)
 {
 	Description = NSLOCTEXT("UnrealEd", "Split", "Split");
 	bPushButton = true;
@@ -2759,8 +2759,8 @@ bool UGeomModifier_Split::OnApply()
 	return true;
 }
 
-UGeomModifier_Triangulate::UGeomModifier_Triangulate(const class FPostConstructInitializeProperties& PCIP)
-	: Super(PCIP)
+UGeomModifier_Triangulate::UGeomModifier_Triangulate(const FObjectInitializer& ObjectInitializer)
+	: Super(ObjectInitializer)
 {
 	Description = NSLOCTEXT("UnrealEd", "Triangulate", "Triangulate");
 	bPushButton = true;
@@ -2830,8 +2830,8 @@ bool UGeomModifier_Triangulate::OnApply()
 	return true;
 }
 
-UGeomModifier_Optimize::UGeomModifier_Optimize(const class FPostConstructInitializeProperties& PCIP)
-	: Super(PCIP)
+UGeomModifier_Optimize::UGeomModifier_Optimize(const FObjectInitializer& ObjectInitializer)
+	: Super(ObjectInitializer)
 {
 	Description = NSLOCTEXT("UnrealEd", "Optimize", "Optimize");
 	bPushButton = true;
@@ -2917,8 +2917,8 @@ bool UGeomModifier_Optimize::OnApply()
 	return true;
 }
 
-UGeomModifier_Turn::UGeomModifier_Turn(const class FPostConstructInitializeProperties& PCIP)
-	: Super(PCIP)
+UGeomModifier_Turn::UGeomModifier_Turn(const FObjectInitializer& ObjectInitializer)
+	: Super(ObjectInitializer)
 {
 	Description = NSLOCTEXT("UnrealEd", "Turn", "Turn");
 	bPushButton = true;
@@ -3077,8 +3077,8 @@ bool UGeomModifier_Turn::OnApply()
 	return true;
 }
 
-UGeomModifier_Weld::UGeomModifier_Weld(const class FPostConstructInitializeProperties& PCIP)
-	: Super(PCIP)
+UGeomModifier_Weld::UGeomModifier_Weld(const FObjectInitializer& ObjectInitializer)
+	: Super(ObjectInitializer)
 {
 	Description = NSLOCTEXT("UnrealEd", "Weld", "Weld");
 	bPushButton = true;
diff --git a/Engine/Source/Editor/GraphEditor/Private/GraphEditorSettings.cpp b/Engine/Source/Editor/GraphEditor/Private/GraphEditorSettings.cpp
--- a/Engine/Source/Editor/GraphEditor/Private/GraphEditorSettings.cpp
+++ b/Engine/Source/Editor/GraphEditor/Private/GraphEditorSettings.cpp
@@ -6,8 +6,8 @@
 /* UGraphEditorSettings structors
  *****************************************************************************/
 
-UGraphEditorSettings::UGraphEditorSettings( const class FPostConstructInitializeProperties& PCIP )
-	: Super(PCIP)
+UGraphEditorSettings::UGraphEditorSettings( const FObjectInitializer& ObjectInitializer )
+	: Super(ObjectInitializer)
 	, PaddingAbovePin(4.0f)
 	, PaddingBelowPin(4.0f)
 	, PaddingRightOfInput(10.0f)
diff --git a/Engine/Source/Editor/HardwareTargeting/Private/HardwareTargetingSettings.cpp b/Engine/Source/Editor/HardwareTargeting/Private/HardwareTargetingSettings.cpp
--- a/Engine/Source/Editor/HardwareTargeting/Private/HardwareTargetingSettings.cpp
+++ b/Engine/Source/Editor/HardwareTargeting/Private/HardwareTargetingSettings.cpp
@@ -7,8 +7,8 @@
 #include "Internationalization.h"
 #include "SDecoratedEnumCombo.h"
 
-UHardwareTargetingSettings::UHardwareTargetingSettings(const FPostConstructInitializeProperties& PCIP)
-	: Super(PCIP)
+UHardwareTargetingSettings::UHardwareTargetingSettings(const FObjectInitializer& ObjectInitializer)
+	: Super(ObjectInitializer)
 	, TargetedHardwareClass(EHardwareClass::Unspecified)
 	, AppliedTargetedHardwareClass(EHardwareClass::Unspecified)
 	, DefaultGraphicsPerformance(EGraphicsPreset::Unspecified)
diff --git a/Engine/Source/Editor/IntroTutorials/Private/EditorTutorial.cpp b/Engine/Source/Editor/IntroTutorials/Private/EditorTutorial.cpp
--- a/Engine/Source/Editor/IntroTutorials/Private/EditorTutorial.cpp
+++ b/Engine/Source/Editor/IntroTutorials/Private/EditorTutorial.cpp
@@ -3,8 +3,8 @@
 #include "IntroTutorialsPrivatePCH.h"
 #include "EditorTutorial.h"
 
-UEditorTutorial::UEditorTutorial(const class FPostConstructInitializeProperties& PCIP)
-	: Super(PCIP)
+UEditorTutorial::UEditorTutorial(const FObjectInitializer& ObjectInitializer)
+	: Super(ObjectInitializer)
 {
 }
 
diff --git a/Engine/Source/Editor/IntroTutorials/Private/EditorTutorialFactory.cpp b/Engine/Source/Editor/IntroTutorials/Private/EditorTutorialFactory.cpp
--- a/Engine/Source/Editor/IntroTutorials/Private/EditorTutorialFactory.cpp
+++ b/Engine/Source/Editor/IntroTutorials/Private/EditorTutorialFactory.cpp
@@ -8,8 +8,8 @@
 
 #define LOCTEXT_NAMESPACE "UEditorTutorialFactory"
 
-UEditorTutorialFactory::UEditorTutorialFactory(const class FPostConstructInitializeProperties& PCIP)
-	: Super(PCIP)
+UEditorTutorialFactory::UEditorTutorialFactory(const FObjectInitializer& ObjectInitializer)
+	: Super(ObjectInitializer)
 {
 	bCreateNew = true;
 	bEditAfterNew = true;
diff --git a/Engine/Source/Editor/IntroTutorials/Private/EditorTutorialImportFactory.cpp b/Engine/Source/Editor/IntroTutorials/Private/EditorTutorialImportFactory.cpp
--- a/Engine/Source/Editor/IntroTutorials/Private/EditorTutorialImportFactory.cpp
+++ b/Engine/Source/Editor/IntroTutorials/Private/EditorTutorialImportFactory.cpp
@@ -10,8 +10,8 @@
 
 #define LOCTEXT_NAMESPACE "UEditorTutorialImportFactory"
 
-UEditorTutorialImportFactory::UEditorTutorialImportFactory(const class FPostConstructInitializeProperties& PCIP)
-	: Super(PCIP)
+UEditorTutorialImportFactory::UEditorTutorialImportFactory(const FObjectInitializer& ObjectInitializer)
+	: Super(ObjectInitializer)
 {
 	bEditorImport = true;
 	bEditAfterNew = true;
diff --git a/Engine/Source/Editor/IntroTutorials/Private/EditorTutorialSettings.cpp b/Engine/Source/Editor/IntroTutorials/Private/EditorTutorialSettings.cpp
--- a/Engine/Source/Editor/IntroTutorials/Private/EditorTutorialSettings.cpp
+++ b/Engine/Source/Editor/IntroTutorials/Private/EditorTutorialSettings.cpp
@@ -3,8 +3,8 @@
 #include "IntroTutorialsPrivatePCH.h"
 #include "EditorTutorialSettings.h"
 
-UEditorTutorialSettings::UEditorTutorialSettings(const class FPostConstructInitializeProperties& PCIP)
-	: Super(PCIP)
+UEditorTutorialSettings::UEditorTutorialSettings(const FObjectInitializer& ObjectInitializer)
+	: Super(ObjectInitializer)
 {
 
 }
diff --git a/Engine/Source/Editor/IntroTutorials/Private/TutorialSettings.cpp b/Engine/Source/Editor/IntroTutorials/Private/TutorialSettings.cpp
--- a/Engine/Source/Editor/IntroTutorials/Private/TutorialSettings.cpp
+++ b/Engine/Source/Editor/IntroTutorials/Private/TutorialSettings.cpp
@@ -3,8 +3,8 @@
 #include "IntroTutorialsPrivatePCH.h"
 #include "TutorialSettings.h"
 
-UTutorialSettings::UTutorialSettings(const class FPostConstructInitializeProperties& PCIP)
-	: Super(PCIP)
+UTutorialSettings::UTutorialSettings(const FObjectInitializer& ObjectInitializer)
+	: Super(ObjectInitializer)
 {
 }
 
diff --git a/Engine/Source/Editor/IntroTutorials/Private/TutorialStateSettings.cpp b/Engine/Source/Editor/IntroTutorials/Private/TutorialStateSettings.cpp
--- a/Engine/Source/Editor/IntroTutorials/Private/TutorialStateSettings.cpp
+++ b/Engine/Source/Editor/IntroTutorials/Private/TutorialStateSettings.cpp
@@ -3,8 +3,8 @@
 #include "IntroTutorialsPrivatePCH.h"
 #include "TutorialStateSettings.h"
 
-UTutorialStateSettings::UTutorialStateSettings(const class FPostConstructInitializeProperties& PCIP)
-	: Super(PCIP)
+UTutorialStateSettings::UTutorialStateSettings(const FObjectInitializer& ObjectInitializer)
+	: Super(ObjectInitializer)
 {
 
 }
diff --git a/Engine/Source/Editor/Kismet/Private/BlueprintPaletteFavorites.cpp b/Engine/Source/Editor/Kismet/Private/BlueprintPaletteFavorites.cpp
--- a/Engine/Source/Editor/Kismet/Private/BlueprintPaletteFavorites.cpp
+++ b/Engine/Source/Editor/Kismet/Private/BlueprintPaletteFavorites.cpp
@@ -209,8 +209,8 @@ FString const& FFavoritedBlueprintPaletteItem::ToString() const
  ******************************************************************************/
 
 //------------------------------------------------------------------------------
-UBlueprintPaletteFavorites::UBlueprintPaletteFavorites(class FPostConstructInitializeProperties const& PCIP)
-	: Super(PCIP)
+UBlueprintPaletteFavorites::UBlueprintPaletteFavorites(FObjectInitializer const& ObjectInitializer)
+	: Super(ObjectInitializer)
 {
 }
 
diff --git a/Engine/Source/Editor/KismetCompiler/Private/KismetCompilerCppBackend.cpp b/Engine/Source/Editor/KismetCompiler/Private/KismetCompilerCppBackend.cpp
--- a/Engine/Source/Editor/KismetCompiler/Private/KismetCompilerCppBackend.cpp
+++ b/Engine/Source/Editor/KismetCompiler/Private/KismetCompilerCppBackend.cpp
@@ -639,7 +639,7 @@ void FKismetCppBackend::GenerateCodeFromClass(UClass* SourceClass, TIndirectArra
 
 	//constructor
 	Emit(Body, *FString::Printf(
-		TEXT("%s::%s(const class FPostConstructInitializeProperties& PCIP) : Super(PCIP) {}\n\n"),
+		TEXT("%s::%s(const FObjectInitializer& ObjectInitializer) : Super(ObjectInitializer) {}\n\n"),
 		*CppClassName, *CppClassName));
 
 	for (int32 i = 0; i < Functions.Num(); ++i)
diff --git a/Engine/Source/Editor/LandscapeEditor/Private/ActorFactoryLandscape.cpp b/Engine/Source/Editor/LandscapeEditor/Private/ActorFactoryLandscape.cpp
--- a/Engine/Source/Editor/LandscapeEditor/Private/ActorFactoryLandscape.cpp
+++ b/Engine/Source/Editor/LandscapeEditor/Private/ActorFactoryLandscape.cpp
@@ -6,8 +6,8 @@
 
 #define LOCTEXT_NAMESPACE "Landscape"
 
-UActorFactoryLandscape::UActorFactoryLandscape(const FPostConstructInitializeProperties& PCIP)
-	: Super(PCIP)
+UActorFactoryLandscape::UActorFactoryLandscape(const FObjectInitializer& ObjectInitializer)
+	: Super(ObjectInitializer)
 {
 	DisplayName = LOCTEXT("Landscape", "Landscape");
 	NewActorClass = ALandscapeProxy::StaticClass();
@@ -31,8 +31,8 @@ AActor* UActorFactoryLandscape::SpawnActor(UObject* Asset, ULevel* InLevel, cons
 	return InLevel->OwningWorld->SpawnActor(ALandscapePlaceholder::StaticClass(), &Location, &Rotation, SpawnInfo);
 }
 
-ALandscapePlaceholder::ALandscapePlaceholder(const FPostConstructInitializeProperties& PCIP)
-	: Super(PCIP)
+ALandscapePlaceholder::ALandscapePlaceholder(const FObjectInitializer& ObjectInitializer)
+	: Super(ObjectInitializer)
 {
 	// Structure to hold one-time initialization
 	struct FConstructorStatics
@@ -45,12 +45,12 @@ ALandscapePlaceholder::ALandscapePlaceholder(const FPostConstructInitializePrope
 	};
 	static FConstructorStatics ConstructorStatics;
 
-	TSubobjectPtr<USceneComponent> SceneComponent = PCIP.CreateDefaultSubobject<USceneComponent>(this, TEXT("RootComponent0"));
+	TSubobjectPtr<USceneComponent> SceneComponent = ObjectInitializer.CreateDefaultSubobject<USceneComponent>(this, TEXT("RootComponent0"));
 	RootComponent = SceneComponent.Get();
 	RootComponent->Mobility = EComponentMobility::Static;
 
 #if WITH_EDITORONLY_DATA
-	TSubobjectPtr<UBillboardComponent> SpriteComponent = PCIP.CreateEditorOnlyDefaultSubobject<UBillboardComponent>(this, TEXT("Sprite"));
+	TSubobjectPtr<UBillboardComponent> SpriteComponent = ObjectInitializer.CreateEditorOnlyDefaultSubobject<UBillboardComponent>(this, TEXT("Sprite"));
 	if (SpriteComponent.IsValid())
 	{
 		SpriteComponent->Sprite = ConstructorStatics.TerrainTexture.Get();
diff --git a/Engine/Source/Editor/LandscapeEditor/Private/LandscapeEditorObject.cpp b/Engine/Source/Editor/LandscapeEditor/Private/LandscapeEditorObject.cpp
--- a/Engine/Source/Editor/LandscapeEditor/Private/LandscapeEditorObject.cpp
+++ b/Engine/Source/Editor/LandscapeEditor/Private/LandscapeEditorObject.cpp
@@ -7,8 +7,8 @@
 
 #include "Landscape/LandscapeMaterialInstanceConstant.h"
 
-ULandscapeEditorObject::ULandscapeEditorObject(const class FPostConstructInitializeProperties& PCIP)
-	: Super(PCIP)
+ULandscapeEditorObject::ULandscapeEditorObject(const FObjectInitializer& ObjectInitializer)
+	: Super(ObjectInitializer)
 
 	// Tool Settings:
 	, ToolStrength(0.3f)
diff --git a/Engine/Source/Editor/LevelEditor/Private/LevelEditorActions.cpp b/Engine/Source/Editor/LevelEditor/Private/LevelEditorActions.cpp
--- a/Engine/Source/Editor/LevelEditor/Private/LevelEditorActions.cpp
+++ b/Engine/Source/Editor/LevelEditor/Private/LevelEditorActions.cpp
@@ -2341,7 +2341,7 @@ void FLevelEditorActionCallbacks::OnSaveBrushAsCollision()
 
 		// Copy the current builder brush into a temp model.
 		// We keep no reference to this, so it will be GC'd at some point.
-		UModel* TempModel = new UModel(FPostConstructInitializeProperties(), NULL, 1);
+		UModel* TempModel = new UModel(FObjectInitializer(), NULL, 1);
 		TempModel->Polys->Element.AssignButKeepOwner(BuilderModel->Polys->Element);
 
 		// Now transform each poly into local space for the selected static mesh.
diff --git a/Engine/Source/Editor/Levels/Private/LevelsModule.cpp b/Engine/Source/Editor/Levels/Private/LevelsModule.cpp
--- a/Engine/Source/Editor/Levels/Private/LevelsModule.cpp
+++ b/Engine/Source/Editor/Levels/Private/LevelsModule.cpp
@@ -29,6 +29,6 @@ TSharedRef< SWidget > FLevelsModule::CreateLevelBrowser()
 }
 
 
-ULevelBrowserSettings::ULevelBrowserSettings(const class FPostConstructInitializeProperties& PCIP)
-	: Super(PCIP)
+ULevelBrowserSettings::ULevelBrowserSettings(const FObjectInitializer& ObjectInitializer)
+	: Super(ObjectInitializer)
 { }
diff --git a/Engine/Source/Editor/Matinee/Classes/MatineeTransBuffer.h b/Engine/Source/Editor/Matinee/Classes/MatineeTransBuffer.h
--- a/Engine/Source/Editor/Matinee/Classes/MatineeTransBuffer.h
+++ b/Engine/Source/Editor/Matinee/Classes/MatineeTransBuffer.h
@@ -14,8 +14,8 @@ class UMatineeTransBuffer : public UTransBuffer
 {
     GENERATED_UCLASS_BODY()
 
-	UMatineeTransBuffer(const class FPostConstructInitializeProperties& PCIP, SIZE_T InMaxMemory)
-		:	UTransBuffer( PCIP, InMaxMemory )
+	UMatineeTransBuffer(const FObjectInitializer& ObjectInitializer, SIZE_T InMaxMemory)
+		:	UTransBuffer( ObjectInitializer, InMaxMemory )
 	{ }
 
 	/**  
diff --git a/Engine/Source/Editor/Matinee/Private/Matinee.cpp b/Engine/Source/Editor/Matinee/Private/Matinee.cpp
--- a/Engine/Source/Editor/Matinee/Private/Matinee.cpp
+++ b/Engine/Source/Editor/Matinee/Private/Matinee.cpp
@@ -485,7 +485,7 @@ void FMatinee::InitMatinee(const EToolkitMode::Type Mode, const TSharedPtr< clas
 	GEditor->ResetTransaction( NSLOCTEXT("UnrealEd", "OpenMatinee", "Open UnrealMatinee") );
 
 	NormalTransactor = GEditor->Trans;
-	InterpEdTrans = new UMatineeTransBuffer( FPostConstructInitializeProperties(), 8*1024*1024 );
+	InterpEdTrans = new UMatineeTransBuffer( FObjectInitializer(), 8*1024*1024 );
 	GEditor->Trans = InterpEdTrans;
 
 	// Save viewports' data before it gets overridden by UpdateLevelViewport
diff --git a/Engine/Source/Editor/Matinee/Private/MatineeOptions.cpp b/Engine/Source/Editor/Matinee/Private/MatineeOptions.cpp
--- a/Engine/Source/Editor/Matinee/Private/MatineeOptions.cpp
+++ b/Engine/Source/Editor/Matinee/Private/MatineeOptions.cpp
@@ -6,8 +6,8 @@
 /*-----------------------------------------------------------------------------
  MatineeOptions
  -----------------------------------------------------------------------------*/
-UMatineeOptions::UMatineeOptions(const class FPostConstructInitializeProperties& PCIP)
-	: Super(PCIP)
+UMatineeOptions::UMatineeOptions(const FObjectInitializer& ObjectInitializer)
+	: Super(ObjectInitializer)
 {
 }
 
diff --git a/Engine/Source/Editor/Matinee/Private/MatineeTrackHelpers.cpp b/Engine/Source/Editor/Matinee/Private/MatineeTrackHelpers.cpp
--- a/Engine/Source/Editor/Matinee/Private/MatineeTrackHelpers.cpp
+++ b/Engine/Source/Editor/Matinee/Private/MatineeTrackHelpers.cpp
@@ -27,8 +27,8 @@
 
 FName	UInterpTrackHelper::KeyframeAddDataName = NAME_None;
 
-UInterpTrackHelper::UInterpTrackHelper(const class FPostConstructInitializeProperties& PCIP)
-	: Super(PCIP)
+UInterpTrackHelper::UInterpTrackHelper(const FObjectInitializer& ObjectInitializer)
+	: Super(ObjectInitializer)
 {
 }
 
@@ -72,8 +72,8 @@ void FMatinee::SetTrackAddPropName( const FName NewName )
 	TrackAddPropName = NewName;
 }
 
-UMatineeTrackAnimControlHelper::UMatineeTrackAnimControlHelper(const class FPostConstructInitializeProperties& PCIP)
-	: Super(PCIP)
+UMatineeTrackAnimControlHelper::UMatineeTrackAnimControlHelper(const FObjectInitializer& ObjectInitializer)
+	: Super(ObjectInitializer)
 {
 }
 
@@ -312,8 +312,8 @@ void  UMatineeTrackAnimControlHelper::PostCreateKeyframe( UInterpTrack *Track, i
 }
 
 
-UMatineeTrackDirectorHelper::UMatineeTrackDirectorHelper(const class FPostConstructInitializeProperties& PCIP)
-	: Super(PCIP)
+UMatineeTrackDirectorHelper::UMatineeTrackDirectorHelper(const FObjectInitializer& ObjectInitializer)
+	: Super(ObjectInitializer)
 {
 }
 
@@ -381,8 +381,8 @@ void  UMatineeTrackDirectorHelper::PostCreateKeyframe( UInterpTrack *Track, int3
 
 
 
-UMatineeTrackEventHelper::UMatineeTrackEventHelper(const class FPostConstructInitializeProperties& PCIP)
-	: Super(PCIP)
+UMatineeTrackEventHelper::UMatineeTrackEventHelper(const FObjectInitializer& ObjectInitializer)
+	: Super(ObjectInitializer)
 {
 }
 
@@ -453,8 +453,8 @@ void  UMatineeTrackEventHelper::PostCreateKeyframe( UInterpTrack *Track, int32 K
 	KeyframeAddDataName = NAME_None;
 }
 
-UMatineeTrackSoundHelper::UMatineeTrackSoundHelper(const class FPostConstructInitializeProperties& PCIP)
-	: Super(PCIP)
+UMatineeTrackSoundHelper::UMatineeTrackSoundHelper(const FObjectInitializer& ObjectInitializer)
+	: Super(ObjectInitializer)
 {
 }
 
@@ -482,8 +482,8 @@ void  UMatineeTrackSoundHelper::PostCreateKeyframe( UInterpTrack *Track, int32 K
 	KeyframeAddSound = NULL;
 }
 
-UMatineeTrackFloatPropHelper::UMatineeTrackFloatPropHelper(const class FPostConstructInitializeProperties& PCIP)
-	: Super(PCIP)
+UMatineeTrackFloatPropHelper::UMatineeTrackFloatPropHelper(const FObjectInitializer& ObjectInitializer)
+	: Super(ObjectInitializer)
 {
 }
 
@@ -619,8 +619,8 @@ void  UMatineeTrackFloatPropHelper::PostCreateTrack( UInterpTrack *Track, bool b
 
 
 
-UMatineeTrackBoolPropHelper::UMatineeTrackBoolPropHelper(const class FPostConstructInitializeProperties& PCIP)
-	: Super(PCIP)
+UMatineeTrackBoolPropHelper::UMatineeTrackBoolPropHelper(const FObjectInitializer& ObjectInitializer)
+	: Super(ObjectInitializer)
 {
 }
 
@@ -756,8 +756,8 @@ void UMatineeTrackBoolPropHelper::PostCreateTrack( UInterpTrack* Track, bool bDu
 }
 
 
-UMatineeTrackToggleHelper::UMatineeTrackToggleHelper(const class FPostConstructInitializeProperties& PCIP)
-	: Super(PCIP)
+UMatineeTrackToggleHelper::UMatineeTrackToggleHelper(const FObjectInitializer& ObjectInitializer)
+	: Super(ObjectInitializer)
 {
 }
 
@@ -835,8 +835,8 @@ void  UMatineeTrackToggleHelper::PostCreateKeyframe( UInterpTrack *Track, int32
 		KeyframeAddDataName = NAME_None;
 }
 
-UMatineeTrackVectorPropHelper::UMatineeTrackVectorPropHelper(const class FPostConstructInitializeProperties& PCIP)
-	: Super(PCIP)
+UMatineeTrackVectorPropHelper::UMatineeTrackVectorPropHelper(const FObjectInitializer& ObjectInitializer)
+	: Super(ObjectInitializer)
 {
 }
 
@@ -982,8 +982,8 @@ void  UMatineeTrackVectorPropHelper::PostCreateTrack( UInterpTrack *Track, bool
 	}
 }
 
-UMatineeTrackColorPropHelper::UMatineeTrackColorPropHelper(const class FPostConstructInitializeProperties& PCIP)
-	: Super(PCIP)
+UMatineeTrackColorPropHelper::UMatineeTrackColorPropHelper(const FObjectInitializer& ObjectInitializer)
+	: Super(ObjectInitializer)
 {
 }
 
@@ -1047,8 +1047,8 @@ void  UMatineeTrackColorPropHelper::PostCreateTrack( UInterpTrack *Track, bool b
 
 
 
-UMatineeTrackLinearColorPropHelper::UMatineeTrackLinearColorPropHelper(const class FPostConstructInitializeProperties& PCIP)
-	: Super(PCIP)
+UMatineeTrackLinearColorPropHelper::UMatineeTrackLinearColorPropHelper(const FObjectInitializer& ObjectInitializer)
+	: Super(ObjectInitializer)
 {
 }
 
@@ -1110,8 +1110,8 @@ void  UMatineeTrackLinearColorPropHelper::PostCreateTrack( UInterpTrack *Track,
 	}
 }
 
-UMatineeTrackVisibilityHelper::UMatineeTrackVisibilityHelper(const class FPostConstructInitializeProperties& PCIP)
-	: Super(PCIP)
+UMatineeTrackVisibilityHelper::UMatineeTrackVisibilityHelper(const FObjectInitializer& ObjectInitializer)
+	: Super(ObjectInitializer)
 {
 }
 
@@ -1194,8 +1194,8 @@ void  UMatineeTrackVisibilityHelper::PostCreateKeyframe( UInterpTrack *Track, in
 
 
 
-UMatineeTrackParticleReplayHelper::UMatineeTrackParticleReplayHelper(const class FPostConstructInitializeProperties& PCIP)
-	: Super(PCIP)
+UMatineeTrackParticleReplayHelper::UMatineeTrackParticleReplayHelper(const FObjectInitializer& ObjectInitializer)
+	: Super(ObjectInitializer)
 {
 }
 
diff --git a/Engine/Source/Editor/Matinee/Private/MatineeTransBuffer.cpp b/Engine/Source/Editor/Matinee/Private/MatineeTransBuffer.cpp
--- a/Engine/Source/Editor/Matinee/Private/MatineeTransBuffer.cpp
+++ b/Engine/Source/Editor/Matinee/Private/MatineeTransBuffer.cpp
@@ -7,8 +7,8 @@
 /*-----------------------------------------------------------------------------
 	UMatineeTransBuffer / FMatineeTransaction
 -----------------------------------------------------------------------------*/
-UMatineeTransBuffer::UMatineeTransBuffer(const class FPostConstructInitializeProperties& PCIP)
-	: Super(PCIP)
+UMatineeTransBuffer::UMatineeTransBuffer(const FObjectInitializer& ObjectInitializer)
+	: Super(ObjectInitializer)
 {
 }
 
diff --git a/Engine/Source/Editor/NewsFeed/Private/NewsFeedClasses.cpp b/Engine/Source/Editor/NewsFeed/Private/NewsFeedClasses.cpp
--- a/Engine/Source/Editor/NewsFeed/Private/NewsFeedClasses.cpp
+++ b/Engine/Source/Editor/NewsFeed/Private/NewsFeedClasses.cpp
@@ -6,8 +6,8 @@
 /* UNewsFeedSettings structors
  *****************************************************************************/
 
-UNewsFeedSettings::UNewsFeedSettings( const class FPostConstructInitializeProperties& PCIP )
-	: Super(PCIP)
+UNewsFeedSettings::UNewsFeedSettings( const FObjectInitializer& ObjectInitializer )
+	: Super(ObjectInitializer)
 	, MaxItemsToShow(10)
 	, ShowOnlyUnreadItems(true)
 { }
diff --git a/Engine/Source/Editor/OnlineBlueprintSupport/Private/K2Node_InAppPurchase.cpp b/Engine/Source/Editor/OnlineBlueprintSupport/Private/K2Node_InAppPurchase.cpp
--- a/Engine/Source/Editor/OnlineBlueprintSupport/Private/K2Node_InAppPurchase.cpp
+++ b/Engine/Source/Editor/OnlineBlueprintSupport/Private/K2Node_InAppPurchase.cpp
@@ -6,8 +6,8 @@
 
 #define LOCTEXT_NAMESPACE "K2Node"
 
-UK2Node_InAppPurchase::UK2Node_InAppPurchase(const FPostConstructInitializeProperties& PCIP)
-	: Super(PCIP)
+UK2Node_InAppPurchase::UK2Node_InAppPurchase(const FObjectInitializer& ObjectInitializer)
+	: Super(ObjectInitializer)
 {
 	ProxyFactoryFunctionName = GET_FUNCTION_NAME_CHECKED(UInAppPurchaseCallbackProxy, CreateProxyObjectForInAppPurchase);
 	ProxyFactoryClass = UInAppPurchaseCallbackProxy::StaticClass();
diff --git a/Engine/Source/Editor/OnlineBlueprintSupport/Private/K2Node_InAppPurchaseQuery.cpp b/Engine/Source/Editor/OnlineBlueprintSupport/Private/K2Node_InAppPurchaseQuery.cpp
--- a/Engine/Source/Editor/OnlineBlueprintSupport/Private/K2Node_InAppPurchaseQuery.cpp
+++ b/Engine/Source/Editor/OnlineBlueprintSupport/Private/K2Node_InAppPurchaseQuery.cpp
@@ -6,8 +6,8 @@
 
 #define LOCTEXT_NAMESPACE "K2Node"
 
-UK2Node_InAppPurchaseQuery::UK2Node_InAppPurchaseQuery(const FPostConstructInitializeProperties& PCIP)
-	: Super(PCIP)
+UK2Node_InAppPurchaseQuery::UK2Node_InAppPurchaseQuery(const FObjectInitializer& ObjectInitializer)
+	: Super(ObjectInitializer)
 {
 	ProxyFactoryFunctionName = GET_FUNCTION_NAME_CHECKED(UInAppPurchaseQueryCallbackProxy, CreateProxyObjectForInAppPurchaseQuery);
 	ProxyFactoryClass = UInAppPurchaseQueryCallbackProxy::StaticClass();
diff --git a/Engine/Source/Editor/OnlineBlueprintSupport/Private/K2Node_LatentOnlineCall.cpp b/Engine/Source/Editor/OnlineBlueprintSupport/Private/K2Node_LatentOnlineCall.cpp
--- a/Engine/Source/Editor/OnlineBlueprintSupport/Private/K2Node_LatentOnlineCall.cpp
+++ b/Engine/Source/Editor/OnlineBlueprintSupport/Private/K2Node_LatentOnlineCall.cpp
@@ -8,8 +8,8 @@
 
 #define LOCTEXT_NAMESPACE "K2Node"
 
-UK2Node_LatentOnlineCall::UK2Node_LatentOnlineCall(const FPostConstructInitializeProperties& PCIP)
-	: Super(PCIP)
+UK2Node_LatentOnlineCall::UK2Node_LatentOnlineCall(const FObjectInitializer& ObjectInitializer)
+	: Super(ObjectInitializer)
 {
 	ProxyActivateFunctionName = GET_FUNCTION_NAME_CHECKED(UOnlineBlueprintCallProxyBase, Activate);
 }
diff --git a/Engine/Source/Editor/OnlineBlueprintSupport/Private/K2Node_LeaderboardFlush.cpp b/Engine/Source/Editor/OnlineBlueprintSupport/Private/K2Node_LeaderboardFlush.cpp
--- a/Engine/Source/Editor/OnlineBlueprintSupport/Private/K2Node_LeaderboardFlush.cpp
+++ b/Engine/Source/Editor/OnlineBlueprintSupport/Private/K2Node_LeaderboardFlush.cpp
@@ -6,8 +6,8 @@
 
 #define LOCTEXT_NAMESPACE "K2Node"
 
-UK2Node_LeaderboardFlush::UK2Node_LeaderboardFlush(const FPostConstructInitializeProperties& PCIP)
-	: Super(PCIP)
+UK2Node_LeaderboardFlush::UK2Node_LeaderboardFlush(const FObjectInitializer& ObjectInitializer)
+	: Super(ObjectInitializer)
 {
 	ProxyFactoryFunctionName = GET_FUNCTION_NAME_CHECKED(ULeaderboardFlushCallbackProxy, CreateProxyObjectForFlush);
 	ProxyFactoryClass = ULeaderboardFlushCallbackProxy::StaticClass();
diff --git a/Engine/Source/Editor/OnlineBlueprintSupport/Private/K2Node_LeaderboardQuery.cpp b/Engine/Source/Editor/OnlineBlueprintSupport/Private/K2Node_LeaderboardQuery.cpp
--- a/Engine/Source/Editor/OnlineBlueprintSupport/Private/K2Node_LeaderboardQuery.cpp
+++ b/Engine/Source/Editor/OnlineBlueprintSupport/Private/K2Node_LeaderboardQuery.cpp
@@ -6,8 +6,8 @@
 
 #define LOCTEXT_NAMESPACE "K2Node"
 
-UK2Node_LeaderboardQuery::UK2Node_LeaderboardQuery(const FPostConstructInitializeProperties& PCIP)
-	: Super(PCIP)
+UK2Node_LeaderboardQuery::UK2Node_LeaderboardQuery(const FObjectInitializer& ObjectInitializer)
+	: Super(ObjectInitializer)
 {
 	ProxyFactoryFunctionName = GET_FUNCTION_NAME_CHECKED(ULeaderboardQueryCallbackProxy, CreateProxyObjectForIntQuery);
 	ProxyFactoryClass = ULeaderboardQueryCallbackProxy::StaticClass();
diff --git a/Engine/Source/Editor/PhAT/Private/PhATRender.cpp b/Engine/Source/Editor/PhAT/Private/PhATRender.cpp
--- a/Engine/Source/Editor/PhAT/Private/PhATRender.cpp
+++ b/Engine/Source/Editor/PhAT/Private/PhATRender.cpp
@@ -6,8 +6,8 @@
 #include "PhATEdSkeletalMeshComponent.h"
 
 
-UPhATEdSkeletalMeshComponent::UPhATEdSkeletalMeshComponent(const class FPostConstructInitializeProperties& PCIP)
-	: Super(PCIP)
+UPhATEdSkeletalMeshComponent::UPhATEdSkeletalMeshComponent(const FObjectInitializer& ObjectInitializer)
+	: Super(ObjectInitializer)
 	, BoneUnselectedColor(170,155,225)
 	, BoneSelectedColor(185,70,0)
 	, ElemSelectedColor(255,166,0)
diff --git a/Engine/Source/Editor/ReferenceViewer/Private/EdGraphNode_Reference.cpp b/Engine/Source/Editor/ReferenceViewer/Private/EdGraphNode_Reference.cpp
--- a/Engine/Source/Editor/ReferenceViewer/Private/EdGraphNode_Reference.cpp
+++ b/Engine/Source/Editor/ReferenceViewer/Private/EdGraphNode_Reference.cpp
@@ -9,8 +9,8 @@
 //////////////////////////////////////////////////////////////////////////
 // UEdGraphNode_Reference
 
-UEdGraphNode_Reference::UEdGraphNode_Reference(const class FPostConstructInitializeProperties& PCIP)
-	: Super(PCIP)
+UEdGraphNode_Reference::UEdGraphNode_Reference(const FObjectInitializer& ObjectInitializer)
+	: Super(ObjectInitializer)
 {
 	DependencyPin = NULL;
 	ReferencerPin = NULL;
diff --git a/Engine/Source/Editor/ReferenceViewer/Private/EdGraph_ReferenceViewer.cpp b/Engine/Source/Editor/ReferenceViewer/Private/EdGraph_ReferenceViewer.cpp
--- a/Engine/Source/Editor/ReferenceViewer/Private/EdGraph_ReferenceViewer.cpp
+++ b/Engine/Source/Editor/ReferenceViewer/Private/EdGraph_ReferenceViewer.cpp
@@ -4,8 +4,8 @@
 #include "AssetRegistryModule.h"
 #include "AssetThumbnail.h"
 
-UEdGraph_ReferenceViewer::UEdGraph_ReferenceViewer(const class FPostConstructInitializeProperties& PCIP)
-	: Super(PCIP)
+UEdGraph_ReferenceViewer::UEdGraph_ReferenceViewer(const FObjectInitializer& ObjectInitializer)
+	: Super(ObjectInitializer)
 {
 	AssetThumbnailPool = MakeShareable( new FAssetThumbnailPool(1024) );
 
diff --git a/Engine/Source/Editor/ReferenceViewer/Private/ReferenceViewerSchema.cpp b/Engine/Source/Editor/ReferenceViewer/Private/ReferenceViewerSchema.cpp
--- a/Engine/Source/Editor/ReferenceViewer/Private/ReferenceViewerSchema.cpp
+++ b/Engine/Source/Editor/ReferenceViewer/Private/ReferenceViewerSchema.cpp
@@ -8,8 +8,8 @@
 //////////////////////////////////////////////////////////////////////////
 // UReferenceViewerSchema
 
-UReferenceViewerSchema::UReferenceViewerSchema(const class FPostConstructInitializeProperties& PCIP)
-	: Super(PCIP)
+UReferenceViewerSchema::UReferenceViewerSchema(const FObjectInitializer& ObjectInitializer)
+	: Super(ObjectInitializer)
 {
 }
 
diff --git a/Engine/Source/Editor/SceneOutliner/Private/SceneOutlinerModule.cpp b/Engine/Source/Editor/SceneOutliner/Private/SceneOutlinerModule.cpp
--- a/Engine/Source/Editor/SceneOutliner/Private/SceneOutlinerModule.cpp
+++ b/Engine/Source/Editor/SceneOutliner/Private/SceneOutlinerModule.cpp
@@ -42,8 +42,8 @@ TSharedRef< ISceneOutliner > FSceneOutlinerModule::CreateSceneOutliner( const FS
 /* Class constructors
  *****************************************************************************/
 
-USceneOutlinerSettings::USceneOutlinerSettings( const class FPostConstructInitializeProperties& PCIP )
-	: Super(PCIP)
+USceneOutlinerSettings::USceneOutlinerSettings( const FObjectInitializer& ObjectInitializer )
+	: Super(ObjectInitializer)
 { }
 
 
diff --git a/Engine/Source/Editor/StatsViewer/Private/StatsEntries/CookerStats.cpp b/Engine/Source/Editor/StatsViewer/Private/StatsEntries/CookerStats.cpp
--- a/Engine/Source/Editor/StatsViewer/Private/StatsEntries/CookerStats.cpp
+++ b/Engine/Source/Editor/StatsViewer/Private/StatsEntries/CookerStats.cpp
@@ -4,6 +4,6 @@
 #include "CookerStats.h"
 
 
-UCookerStats::UCookerStats( const class FPostConstructInitializeProperties& PCIP )
-	: Super(PCIP)
+UCookerStats::UCookerStats( const FObjectInitializer& ObjectInitializer )
+	: Super(ObjectInitializer)
 { }
diff --git a/Engine/Source/Editor/StatsViewer/Private/StatsEntries/LightingBuildInfo.cpp b/Engine/Source/Editor/StatsViewer/Private/StatsEntries/LightingBuildInfo.cpp
--- a/Engine/Source/Editor/StatsViewer/Private/StatsEntries/LightingBuildInfo.cpp
+++ b/Engine/Source/Editor/StatsViewer/Private/StatsEntries/LightingBuildInfo.cpp
@@ -3,8 +3,8 @@
 #include "StatsViewerPrivatePCH.h"
 #include "LightingBuildInfo.h"
 
-ULightingBuildInfo::ULightingBuildInfo(const class FPostConstructInitializeProperties& PCIP)
-	: Super(PCIP)
+ULightingBuildInfo::ULightingBuildInfo(const FObjectInitializer& ObjectInitializer)
+	: Super(ObjectInitializer)
 {
 }
 
diff --git a/Engine/Source/Editor/StatsViewer/Private/StatsEntries/PrimitiveStats.cpp b/Engine/Source/Editor/StatsViewer/Private/StatsEntries/PrimitiveStats.cpp
--- a/Engine/Source/Editor/StatsViewer/Private/StatsEntries/PrimitiveStats.cpp
+++ b/Engine/Source/Editor/StatsViewer/Private/StatsEntries/PrimitiveStats.cpp
@@ -3,8 +3,8 @@
 #include "StatsViewerPrivatePCH.h"
 #include "PrimitiveStats.h"
 
-UPrimitiveStats::UPrimitiveStats(const class FPostConstructInitializeProperties& PCIP)
-	: Super(PCIP)
+UPrimitiveStats::UPrimitiveStats(const FObjectInitializer& ObjectInitializer)
+	: Super(ObjectInitializer)
 {
 }
 
diff --git a/Engine/Source/Editor/StatsViewer/Private/StatsEntries/StaticMeshLightingInfo.cpp b/Engine/Source/Editor/StatsViewer/Private/StatsEntries/StaticMeshLightingInfo.cpp
--- a/Engine/Source/Editor/StatsViewer/Private/StatsEntries/StaticMeshLightingInfo.cpp
+++ b/Engine/Source/Editor/StatsViewer/Private/StatsEntries/StaticMeshLightingInfo.cpp
@@ -5,8 +5,8 @@
 
 #define LOCTEXT_NAMESPACE "Editor.StatsViewer"
 
-UStaticMeshLightingInfo::UStaticMeshLightingInfo(const class FPostConstructInitializeProperties& PCIP)
-	: Super(PCIP)
+UStaticMeshLightingInfo::UStaticMeshLightingInfo(const FObjectInitializer& ObjectInitializer)
+	: Super(ObjectInitializer)
 {
 }
 
diff --git a/Engine/Source/Editor/StatsViewer/Private/StatsEntries/TextureStats.cpp b/Engine/Source/Editor/StatsViewer/Private/StatsEntries/TextureStats.cpp
--- a/Engine/Source/Editor/StatsViewer/Private/StatsEntries/TextureStats.cpp
+++ b/Engine/Source/Editor/StatsViewer/Private/StatsEntries/TextureStats.cpp
@@ -3,8 +3,8 @@
 #include "StatsViewerPrivatePCH.h"
 #include "TextureStats.h"
 
-UTextureStats::UTextureStats(const class FPostConstructInitializeProperties& PCIP) :
-	Super(PCIP),
+UTextureStats::UTextureStats(const FObjectInitializer& ObjectInitializer) :
+	Super(ObjectInitializer),
 	LastTimeRendered( FLT_MAX )
 {
 }
diff --git a/Engine/Source/Editor/TextureEditor/Private/TextureEditorClasses.cpp b/Engine/Source/Editor/TextureEditor/Private/TextureEditorClasses.cpp
--- a/Engine/Source/Editor/TextureEditor/Private/TextureEditorClasses.cpp
+++ b/Engine/Source/Editor/TextureEditor/Private/TextureEditorClasses.cpp
@@ -3,8 +3,8 @@
 #include "TextureEditorPrivatePCH.h"
 
 
-UTextureEditorSettings::UTextureEditorSettings( const class FPostConstructInitializeProperties& PCIP )
-	: Super(PCIP)
+UTextureEditorSettings::UTextureEditorSettings( const FObjectInitializer& ObjectInitializer )
+	: Super(ObjectInitializer)
 	, Background(TextureEditorBackground_SolidColor)
 	, BackgroundColor(FColor::Black)
 	, CheckerColorOne(FColor(128, 128, 128))
diff --git a/Engine/Source/Editor/TranslationEditor/Private/InternationalizationExportSettings.cpp b/Engine/Source/Editor/TranslationEditor/Private/InternationalizationExportSettings.cpp
--- a/Engine/Source/Editor/TranslationEditor/Private/InternationalizationExportSettings.cpp
+++ b/Engine/Source/Editor/TranslationEditor/Private/InternationalizationExportSettings.cpp
@@ -3,8 +3,8 @@
 
 #include "InternationalizationExportSettings.h"
 
-UInternationalizationExportSettings::UInternationalizationExportSettings(const class FPostConstructInitializeProperties& PCIP)
-	: Super(PCIP)
+UInternationalizationExportSettings::UInternationalizationExportSettings(const FObjectInitializer& ObjectInitializer)
+	: Super(ObjectInitializer)
 {
 
 }
\ No newline at end of file
diff --git a/Engine/Source/Editor/TranslationEditor/Private/TranslationUnit.cpp b/Engine/Source/Editor/TranslationEditor/Private/TranslationUnit.cpp
--- a/Engine/Source/Editor/TranslationEditor/Private/TranslationUnit.cpp
+++ b/Engine/Source/Editor/TranslationEditor/Private/TranslationUnit.cpp
@@ -2,8 +2,8 @@
 #include "TranslationEditorPrivatePCH.h"
 #include "TranslationUnit.h"
 
-UTranslationUnit::UTranslationUnit( const class FPostConstructInitializeProperties& PCIP )
-	: Super(PCIP)
+UTranslationUnit::UTranslationUnit( const FObjectInitializer& ObjectInitializer )
+	: Super(ObjectInitializer)
 {
 
 }
diff --git a/Engine/Source/Editor/UMGEditor/Private/Nodes/K2Node_CreateDragDropOperation.cpp b/Engine/Source/Editor/UMGEditor/Private/Nodes/K2Node_CreateDragDropOperation.cpp
--- a/Engine/Source/Editor/UMGEditor/Private/Nodes/K2Node_CreateDragDropOperation.cpp
+++ b/Engine/Source/Editor/UMGEditor/Private/Nodes/K2Node_CreateDragDropOperation.cpp
@@ -11,8 +11,8 @@
 
 #define LOCTEXT_NAMESPACE "UMG"
 
-UK2Node_CreateDragDropOperation::UK2Node_CreateDragDropOperation(const class FPostConstructInitializeProperties& PCIP)
-	: Super(PCIP)
+UK2Node_CreateDragDropOperation::UK2Node_CreateDragDropOperation(const FObjectInitializer& ObjectInitializer)
+	: Super(ObjectInitializer)
 {
 	NodeTooltip = LOCTEXT("DragDropNodeTooltip", "Creates a new drag drop operation");
 }
diff --git a/Engine/Source/Editor/UMGEditor/Private/Nodes/K2Node_CreateWidget.cpp b/Engine/Source/Editor/UMGEditor/Private/Nodes/K2Node_CreateWidget.cpp
--- a/Engine/Source/Editor/UMGEditor/Private/Nodes/K2Node_CreateWidget.cpp
+++ b/Engine/Source/Editor/UMGEditor/Private/Nodes/K2Node_CreateWidget.cpp
@@ -18,8 +18,8 @@ struct FK2Node_CreateWidgetHelper
 
 FString FK2Node_CreateWidgetHelper::OwningPlayerPinName(TEXT("OwningPlayer"));
 
-UK2Node_CreateWidget::UK2Node_CreateWidget(const class FPostConstructInitializeProperties& PCIP)
-	: Super(PCIP)
+UK2Node_CreateWidget::UK2Node_CreateWidget(const FObjectInitializer& ObjectInitializer)
+	: Super(ObjectInitializer)
 {
 	NodeTooltip = LOCTEXT("NodeTooltip", "Creates a new widget");
 }
diff --git a/Engine/Source/Editor/UMGEditor/Private/Settings/WidgetDesignerSettings.cpp b/Engine/Source/Editor/UMGEditor/Private/Settings/WidgetDesignerSettings.cpp
--- a/Engine/Source/Editor/UMGEditor/Private/Settings/WidgetDesignerSettings.cpp
+++ b/Engine/Source/Editor/UMGEditor/Private/Settings/WidgetDesignerSettings.cpp
@@ -2,8 +2,8 @@
 
 #include "UMGEditorPrivatePCH.h"
 
-UWidgetDesignerSettings::UWidgetDesignerSettings(const class FPostConstructInitializeProperties& PCIP)
-	: Super(PCIP)
+UWidgetDesignerSettings::UWidgetDesignerSettings(const FObjectInitializer& ObjectInitializer)
+	: Super(ObjectInitializer)
 {
 	GridSnapEnabled = true;
 	GridSnapSize = 4;
diff --git a/Engine/Source/Editor/UMGEditor/Private/WidgetBlueprint.cpp b/Engine/Source/Editor/UMGEditor/Private/WidgetBlueprint.cpp
--- a/Engine/Source/Editor/UMGEditor/Private/WidgetBlueprint.cpp
+++ b/Engine/Source/Editor/UMGEditor/Private/WidgetBlueprint.cpp
@@ -40,8 +40,8 @@ bool FWidgetAnimation_DEPRECATED::SerializeFromMismatchedTag(struct FPropertyTag
 /////////////////////////////////////////////////////
 // UWidgetBlueprint
 
-UWidgetBlueprint::UWidgetBlueprint(const FPostConstructInitializeProperties& PCIP)
-	: Super(PCIP)
+UWidgetBlueprint::UWidgetBlueprint(const FObjectInitializer& ObjectInitializer)
+	: Super(ObjectInitializer)
 {
 	WidgetTree = ConstructObject<UWidgetTree>(UWidgetTree::StaticClass(), this);
 	WidgetTree->SetFlags(RF_Transactional);
diff --git a/Engine/Source/Editor/UMGEditor/Private/WidgetBlueprintFactory.cpp b/Engine/Source/Editor/UMGEditor/Private/WidgetBlueprintFactory.cpp
--- a/Engine/Source/Editor/UMGEditor/Private/WidgetBlueprintFactory.cpp
+++ b/Engine/Source/Editor/UMGEditor/Private/WidgetBlueprintFactory.cpp
@@ -12,8 +12,8 @@
 	UWidgetBlueprintFactory implementation.
 ------------------------------------------------------------------------------*/
 
-UWidgetBlueprintFactory::UWidgetBlueprintFactory(const class FPostConstructInitializeProperties& PCIP)
-	: Super(PCIP)
+UWidgetBlueprintFactory::UWidgetBlueprintFactory(const FObjectInitializer& ObjectInitializer)
+	: Super(ObjectInitializer)
 {
 	bCreateNew = true;
 	bEditAfterNew = true;
diff --git a/Engine/Source/Editor/UMGEditor/Private/WidgetGraphSchema.cpp b/Engine/Source/Editor/UMGEditor/Private/WidgetGraphSchema.cpp
--- a/Engine/Source/Editor/UMGEditor/Private/WidgetGraphSchema.cpp
+++ b/Engine/Source/Editor/UMGEditor/Private/WidgetGraphSchema.cpp
@@ -3,8 +3,8 @@
 #include "UMGEditorPrivatePCH.h"
 #include "WidgetGraphSchema.h"
 
-UWidgetGraphSchema::UWidgetGraphSchema(const FPostConstructInitializeProperties& PCIP)
-: Super(PCIP)
+UWidgetGraphSchema::UWidgetGraphSchema(const FObjectInitializer& ObjectInitializer)
+: Super(ObjectInitializer)
 {
 	NAME_NeverAsPin = TEXT("NeverAsPin");
 	NAME_PinHiddenByDefault = TEXT("PinHiddenByDefault");
diff --git a/Engine/Source/Editor/UnrealEd/Classes/Editor/TransBuffer.h b/Engine/Source/Editor/UnrealEd/Classes/Editor/TransBuffer.h
--- a/Engine/Source/Editor/UnrealEd/Classes/Editor/TransBuffer.h
+++ b/Engine/Source/Editor/UnrealEd/Classes/Editor/TransBuffer.h
@@ -34,7 +34,7 @@ class UNREALED_API UTransBuffer
 public:
 
 	// Constructor.
-	UTransBuffer( const class FPostConstructInitializeProperties& PCIP,SIZE_T InMaxMemory );
+	UTransBuffer( const FObjectInitializer& ObjectInitializer,SIZE_T InMaxMemory );
 
 public:
 
diff --git a/Engine/Source/Editor/UnrealEd/Classes/Factories/TextureFactory.h b/Engine/Source/Editor/UnrealEd/Classes/Factories/TextureFactory.h
--- a/Engine/Source/Editor/UnrealEd/Classes/Factories/TextureFactory.h
+++ b/Engine/Source/Editor/UnrealEd/Classes/Factories/TextureFactory.h
@@ -90,7 +90,7 @@ class UNREALED_API UTextureFactory : public UFactory
 	uint32 bUsingExistingSettings:1;
 
 public:
-	UTextureFactory(const class FPostConstructInitializeProperties& PCIP);
+	UTextureFactory(const FObjectInitializer& ObjectInitializer);
 
 	// Begin UObject Interface
 	virtual void PostInitProperties();
diff --git a/Engine/Source/Editor/UnrealEd/Private/Animation/DebugSkelMeshComponent.cpp b/Engine/Source/Editor/UnrealEd/Private/Animation/DebugSkelMeshComponent.cpp
--- a/Engine/Source/Editor/UnrealEd/Private/Animation/DebugSkelMeshComponent.cpp
+++ b/Engine/Source/Editor/UnrealEd/Private/Animation/DebugSkelMeshComponent.cpp
@@ -100,8 +100,8 @@ public:
 //////////////////////////////////////////////////////////////////////////
 // UDebugSkelMeshComponent
 
-UDebugSkelMeshComponent::UDebugSkelMeshComponent(const class FPostConstructInitializeProperties& PCIP)
-	: Super(PCIP)
+UDebugSkelMeshComponent::UDebugSkelMeshComponent(const FObjectInitializer& ObjectInitializer)
+	: Super(ObjectInitializer)
 {
 	bDrawMesh = true;
 	PreviewInstance = NULL;
diff --git a/Engine/Source/Editor/UnrealEd/Private/Animation/EditorAnimBaseObj.cpp b/Engine/Source/Editor/UnrealEd/Private/Animation/EditorAnimBaseObj.cpp
--- a/Engine/Source/Editor/UnrealEd/Private/Animation/EditorAnimBaseObj.cpp
+++ b/Engine/Source/Editor/UnrealEd/Private/Animation/EditorAnimBaseObj.cpp
@@ -9,8 +9,8 @@
 
 #define LOCTEXT_NAMESPACE "SSkeletonTree"
 
-UEditorAnimBaseObj::UEditorAnimBaseObj(const class FPostConstructInitializeProperties& PCIP)
-	: Super(PCIP)
+UEditorAnimBaseObj::UEditorAnimBaseObj(const FObjectInitializer& ObjectInitializer)
+	: Super(ObjectInitializer)
 {
 }
 
diff --git a/Engine/Source/Editor/UnrealEd/Private/Animation/EditorAnimCompositeSegment.cpp b/Engine/Source/Editor/UnrealEd/Private/Animation/EditorAnimCompositeSegment.cpp
--- a/Engine/Source/Editor/UnrealEd/Private/Animation/EditorAnimCompositeSegment.cpp
+++ b/Engine/Source/Editor/UnrealEd/Private/Animation/EditorAnimCompositeSegment.cpp
@@ -8,8 +8,8 @@
 
 
 
-UEditorAnimCompositeSegment::UEditorAnimCompositeSegment(const class FPostConstructInitializeProperties& PCIP)
-	: Super(PCIP)
+UEditorAnimCompositeSegment::UEditorAnimCompositeSegment(const FObjectInitializer& ObjectInitializer)
+	: Super(ObjectInitializer)
 {
 	AnimSegmentIndex = 0;
 }
diff --git a/Engine/Source/Editor/UnrealEd/Private/Animation/EditorAnimSegment.cpp b/Engine/Source/Editor/UnrealEd/Private/Animation/EditorAnimSegment.cpp
--- a/Engine/Source/Editor/UnrealEd/Private/Animation/EditorAnimSegment.cpp
+++ b/Engine/Source/Editor/UnrealEd/Private/Animation/EditorAnimSegment.cpp
@@ -8,8 +8,8 @@
 
 
 
-UEditorAnimSegment::UEditorAnimSegment(const class FPostConstructInitializeProperties& PCIP)
-	: Super(PCIP)
+UEditorAnimSegment::UEditorAnimSegment(const FObjectInitializer& ObjectInitializer)
+	: Super(ObjectInitializer)
 {
 
 	AnimSlotIndex = 0;
diff --git a/Engine/Source/Editor/UnrealEd/Private/Animation/EditorBranchPoint.cpp b/Engine/Source/Editor/UnrealEd/Private/Animation/EditorBranchPoint.cpp
--- a/Engine/Source/Editor/UnrealEd/Private/Animation/EditorBranchPoint.cpp
+++ b/Engine/Source/Editor/UnrealEd/Private/Animation/EditorBranchPoint.cpp
@@ -8,8 +8,8 @@
 
 
 
-UEditorBranchPoint::UEditorBranchPoint(const class FPostConstructInitializeProperties& PCIP)
-	: Super(PCIP)
+UEditorBranchPoint::UEditorBranchPoint(const FObjectInitializer& ObjectInitializer)
+	: Super(ObjectInitializer)
 {
 
 	BranchIndex = INDEX_NONE;
diff --git a/Engine/Source/Editor/UnrealEd/Private/Animation/EditorCompositeSection.cpp b/Engine/Source/Editor/UnrealEd/Private/Animation/EditorCompositeSection.cpp
--- a/Engine/Source/Editor/UnrealEd/Private/Animation/EditorCompositeSection.cpp
+++ b/Engine/Source/Editor/UnrealEd/Private/Animation/EditorCompositeSection.cpp
@@ -8,8 +8,8 @@
 
 
 
-UEditorCompositeSection::UEditorCompositeSection(const class FPostConstructInitializeProperties& PCIP)
-	: Super(PCIP)
+UEditorCompositeSection::UEditorCompositeSection(const FObjectInitializer& ObjectInitializer)
+	: Super(ObjectInitializer)
 {
 
 	SectionIndex = INDEX_NONE;
diff --git a/Engine/Source/Editor/UnrealEd/Private/Animation/EditorParentPlayerListObj.cpp b/Engine/Source/Editor/UnrealEd/Private/Animation/EditorParentPlayerListObj.cpp
--- a/Engine/Source/Editor/UnrealEd/Private/Animation/EditorParentPlayerListObj.cpp
+++ b/Engine/Source/Editor/UnrealEd/Private/Animation/EditorParentPlayerListObj.cpp
@@ -8,8 +8,8 @@
 #include "BlueprintEditorUtils.h"
 #include "ScopedTransaction.h"
 
-UEditorParentPlayerListObj::UEditorParentPlayerListObj(const class FPostConstructInitializeProperties& PCIP)
-	: Super(PCIP)
+UEditorParentPlayerListObj::UEditorParentPlayerListObj(const FObjectInitializer& ObjectInitializer)
+	: Super(ObjectInitializer)
 	, AnimBlueprint(NULL)
 {
 
diff --git a/Engine/Source/Editor/UnrealEd/Private/Animation/EditorSkeletonNotifyObj.cpp b/Engine/Source/Editor/UnrealEd/Private/Animation/EditorSkeletonNotifyObj.cpp
--- a/Engine/Source/Editor/UnrealEd/Private/Animation/EditorSkeletonNotifyObj.cpp
+++ b/Engine/Source/Editor/UnrealEd/Private/Animation/EditorSkeletonNotifyObj.cpp
@@ -2,8 +2,8 @@
 
 #include "UnrealEd.h"
 
-UEditorSkeletonNotifyObj::UEditorSkeletonNotifyObj(const class FPostConstructInitializeProperties& PCIP)
-	: Super(PCIP)
+UEditorSkeletonNotifyObj::UEditorSkeletonNotifyObj(const FObjectInitializer& ObjectInitializer)
+	: Super(ObjectInitializer)
 {
 }
 
diff --git a/Engine/Source/Editor/UnrealEd/Private/AutoReimportManager.cpp b/Engine/Source/Editor/UnrealEd/Private/AutoReimportManager.cpp
--- a/Engine/Source/Editor/UnrealEd/Private/AutoReimportManager.cpp
+++ b/Engine/Source/Editor/UnrealEd/Private/AutoReimportManager.cpp
@@ -12,8 +12,8 @@
 DEFINE_LOG_CATEGORY_STATIC(LogAutoReimportManager, Log, All);
 
 
-UAutoReimportManager::UAutoReimportManager(const class FPostConstructInitializeProperties& PCIP)
-	: Super(PCIP)
+UAutoReimportManager::UAutoReimportManager(const FObjectInitializer& ObjectInitializer)
+	: Super(ObjectInitializer)
 {
 	ReimportTime = 0.0f;
 }
diff --git a/Engine/Source/Editor/UnrealEd/Private/BSPOps.cpp b/Engine/Source/Editor/UnrealEd/Private/BSPOps.cpp
--- a/Engine/Source/Editor/UnrealEd/Private/BSPOps.cpp
+++ b/Engine/Source/Editor/UnrealEd/Private/BSPOps.cpp
@@ -546,8 +546,8 @@ void FBSPOps::csgCopyBrush( ABrush* Dest, ABrush* Src, uint32 PolyFlags, EObject
 
 	// Duplicate the brush and its polys.
 	Dest->PolyFlags		= PolyFlags;
-	Dest->Brush			= new( Dest, NAME_None, ResFlags )UModel(FPostConstructInitializeProperties(), NULL, Src->Brush->RootOutside );
-	Dest->Brush->Polys	= new( Dest->Brush, NAME_None, ResFlags )UPolys(FPostConstructInitializeProperties());
+	Dest->Brush			= new( Dest, NAME_None, ResFlags )UModel(FObjectInitializer(), NULL, Src->Brush->RootOutside );
+	Dest->Brush->Polys	= new( Dest->Brush, NAME_None, ResFlags )UPolys(FObjectInitializer());
 	check(Dest->Brush->Polys->Element.GetOwner()==Dest->Brush->Polys);
 	Dest->Brush->Polys->Element.AssignButKeepOwner(Src->Brush->Polys->Element);
 	check(Dest->Brush->Polys->Element.GetOwner()==Dest->Brush->Polys);
diff --git a/Engine/Source/Editor/UnrealEd/Private/Commandlets/AutomatedLauncherCommandlet.cpp b/Engine/Source/Editor/UnrealEd/Private/Commandlets/AutomatedLauncherCommandlet.cpp
--- a/Engine/Source/Editor/UnrealEd/Private/Commandlets/AutomatedLauncherCommandlet.cpp
+++ b/Engine/Source/Editor/UnrealEd/Private/Commandlets/AutomatedLauncherCommandlet.cpp
@@ -5,8 +5,8 @@
 #include "LauncherAutomatedService.h"
 
 
-UAutomatedLauncherCommandlet::UAutomatedLauncherCommandlet(const class FPostConstructInitializeProperties& PCIP)
-	: Super(PCIP)
+UAutomatedLauncherCommandlet::UAutomatedLauncherCommandlet(const FObjectInitializer& ObjectInitializer)
+	: Super(ObjectInitializer)
 {
 }
 
diff --git a/Engine/Source/Editor/UnrealEd/Private/Commandlets/ContentCommandlets.cpp b/Engine/Source/Editor/UnrealEd/Private/Commandlets/ContentCommandlets.cpp
--- a/Engine/Source/Editor/UnrealEd/Private/Commandlets/ContentCommandlets.cpp
+++ b/Engine/Source/Editor/UnrealEd/Private/Commandlets/ContentCommandlets.cpp
@@ -1,4 +1,4 @@
-﻿// Copyright 1998-2014 Epic Games, Inc. All Rights Reserved.
+// Copyright 1998-2014 Epic Games, Inc. All Rights Reserved.
 
 /*=============================================================================
 	UCContentCommandlets.cpp: Various commmandlets.
@@ -39,8 +39,8 @@ DEFINE_LOG_CATEGORY_STATIC(LogContentCommandlet, Log, All);
 #define CURRENT_PACKAGE_VERSION 0
 #define IGNORE_PACKAGE_VERSION INDEX_NONE
 
-UResavePackagesCommandlet::UResavePackagesCommandlet(const class FPostConstructInitializeProperties& PCIP)
-	: Super(PCIP)
+UResavePackagesCommandlet::UResavePackagesCommandlet(const FObjectInitializer& ObjectInitializer)
+	: Super(ObjectInitializer)
 {
 }
 
@@ -714,8 +714,8 @@ void UResavePackagesCommandlet::VerboseMessage(const FString& Message)
 	URecompileBlueprintsCommandlet commandlet.
 -----------------------------------------------------------------------------*/
 
-URecompileBlueprintsCommandlet::URecompileBlueprintsCommandlet(const class FPostConstructInitializeProperties& PCIP)
-	: Super(PCIP)
+URecompileBlueprintsCommandlet::URecompileBlueprintsCommandlet(const FObjectInitializer& ObjectInitializer)
+	: Super(ObjectInitializer)
 {
 }
 
@@ -986,8 +986,8 @@ FString MakeCutdownFilename(const FString& Filename, const TCHAR* CutdownDirecto
 	return CutdownDirectory / FPaths::GetCleanFilename(Filename);
 }
 
-UWrangleContentCommandlet::UWrangleContentCommandlet(const class FPostConstructInitializeProperties& PCIP)
-	: Super(PCIP)
+UWrangleContentCommandlet::UWrangleContentCommandlet(const FObjectInitializer& ObjectInitializer)
+	: Super(ObjectInitializer)
 {
 	LogToConsole = false;
 }
@@ -1630,8 +1630,8 @@ int32 UWrangleContentCommandlet::Main( const FString& Params )
 	UListMaterialsUsedWithMeshEmittersCommandlet
 ========================================================================================================== */
 
-UListMaterialsUsedWithMeshEmittersCommandlet::UListMaterialsUsedWithMeshEmittersCommandlet(const class FPostConstructInitializeProperties& PCIP)
-	: Super(PCIP)
+UListMaterialsUsedWithMeshEmittersCommandlet::UListMaterialsUsedWithMeshEmittersCommandlet(const FObjectInitializer& ObjectInitializer)
+	: Super(ObjectInitializer)
 {
 }
 
@@ -1785,8 +1785,8 @@ int32 UListMaterialsUsedWithMeshEmittersCommandlet::Main( const FString& Params
 	UListStaticMeshesImportedFromSpeedTreesCommandlet
 ========================================================================================================== */
 
-UListStaticMeshesImportedFromSpeedTreesCommandlet::UListStaticMeshesImportedFromSpeedTreesCommandlet(const class FPostConstructInitializeProperties& PCIP)
-: Super(PCIP)
+UListStaticMeshesImportedFromSpeedTreesCommandlet::UListStaticMeshesImportedFromSpeedTreesCommandlet(const FObjectInitializer& ObjectInitializer)
+: Super(ObjectInitializer)
 {
 }
 
diff --git a/Engine/Source/Editor/UnrealEd/Private/Commandlets/CookCommandlet.cpp b/Engine/Source/Editor/UnrealEd/Private/Commandlets/CookCommandlet.cpp
--- a/Engine/Source/Editor/UnrealEd/Private/Commandlets/CookCommandlet.cpp
+++ b/Engine/Source/Editor/UnrealEd/Private/Commandlets/CookCommandlet.cpp
@@ -26,8 +26,8 @@
 
 DEFINE_LOG_CATEGORY_STATIC(LogCookCommandlet, Log, All);
 
-UCookerSettings::UCookerSettings(const FPostConstructInitializeProperties& PCIP)
-	: Super(PCIP)
+UCookerSettings::UCookerSettings(const FObjectInitializer& ObjectInitializer)
+	: Super(ObjectInitializer)
 {
 	DefaultPVRTCQuality = 0;
 }
@@ -50,8 +50,8 @@ static FString GetPackageFilename( UPackage* Package )
 /* UCookCommandlet structors
  *****************************************************************************/
 
-UCookCommandlet::UCookCommandlet( const class FPostConstructInitializeProperties& PCIP )
-	: Super(PCIP)
+UCookCommandlet::UCookCommandlet( const class FObjectInitializer& ObjectInitializer )
+	: Super(ObjectInitializer)
 {
 
 	LogToConsole = false;
diff --git a/Engine/Source/Editor/UnrealEd/Private/Commandlets/DerivedDataCacheCommandlet.cpp b/Engine/Source/Editor/UnrealEd/Private/Commandlets/DerivedDataCacheCommandlet.cpp
--- a/Engine/Source/Editor/UnrealEd/Private/Commandlets/DerivedDataCacheCommandlet.cpp
+++ b/Engine/Source/Editor/UnrealEd/Private/Commandlets/DerivedDataCacheCommandlet.cpp
@@ -14,8 +14,8 @@
 
 DEFINE_LOG_CATEGORY_STATIC(LogDerivedDataCacheCommandlet, Log, All);
 
-UDerivedDataCacheCommandlet::UDerivedDataCacheCommandlet(const class FPostConstructInitializeProperties& PCIP)
-	: Super(PCIP)
+UDerivedDataCacheCommandlet::UDerivedDataCacheCommandlet(const FObjectInitializer& ObjectInitializer)
+	: Super(ObjectInitializer)
 {
 	LogToConsole = false;
 }
diff --git a/Engine/Source/Editor/UnrealEd/Private/Commandlets/DiffAssetsCommandlet.cpp b/Engine/Source/Editor/UnrealEd/Private/Commandlets/DiffAssetsCommandlet.cpp
--- a/Engine/Source/Editor/UnrealEd/Private/Commandlets/DiffAssetsCommandlet.cpp
+++ b/Engine/Source/Editor/UnrealEd/Private/Commandlets/DiffAssetsCommandlet.cpp
@@ -9,8 +9,8 @@
 
 DEFINE_LOG_CATEGORY_STATIC(LogDiffAssetsCommandlet, Log, All);
 
-UDiffAssetsCommandlet::UDiffAssetsCommandlet(const class FPostConstructInitializeProperties& PCIP)
-	: Super(PCIP)
+UDiffAssetsCommandlet::UDiffAssetsCommandlet(const FObjectInitializer& ObjectInitializer)
+	: Super(ObjectInitializer)
 {
 }
 
diff --git a/Engine/Source/Editor/UnrealEd/Private/Commandlets/DiffPackagesCommandlet.cpp b/Engine/Source/Editor/UnrealEd/Private/Commandlets/DiffPackagesCommandlet.cpp
--- a/Engine/Source/Editor/UnrealEd/Private/Commandlets/DiffPackagesCommandlet.cpp
+++ b/Engine/Source/Editor/UnrealEd/Private/Commandlets/DiffPackagesCommandlet.cpp
@@ -18,8 +18,8 @@ DEFINE_LOG_CATEGORY_STATIC(LogDiffPackagesCommandlet, Log, All);
 // whether to skip levels when building the initial diff sets (for debugging)
 #define OPTIMIZE_LEVEL_DIFFS 1
 
-UDiffPackagesCommandlet::UDiffPackagesCommandlet(const class FPostConstructInitializeProperties& PCIP)
-	: Super(PCIP)
+UDiffPackagesCommandlet::UDiffPackagesCommandlet(const FObjectInitializer& ObjectInitializer)
+	: Super(ObjectInitializer)
 {
 	LogToConsole = false;
 }
diff --git a/Engine/Source/Editor/UnrealEd/Private/Commandlets/DumpBlueprintsInfoCommandlet.cpp b/Engine/Source/Editor/UnrealEd/Private/Commandlets/DumpBlueprintsInfoCommandlet.cpp
--- a/Engine/Source/Editor/UnrealEd/Private/Commandlets/DumpBlueprintsInfoCommandlet.cpp
+++ b/Engine/Source/Editor/UnrealEd/Private/Commandlets/DumpBlueprintsInfoCommandlet.cpp
@@ -2400,8 +2400,8 @@ static int32 DumpBlueprintInfoUtils::LoadBlueprints(FString AssetName, bool bAll
  ******************************************************************************/
 
 //------------------------------------------------------------------------------
-UDumpBlueprintsInfoCommandlet::UDumpBlueprintsInfoCommandlet(class FPostConstructInitializeProperties const& PCIP)
-	: Super(PCIP)
+UDumpBlueprintsInfoCommandlet::UDumpBlueprintsInfoCommandlet(FObjectInitializer const& ObjectInitializer)
+	: Super(ObjectInitializer)
 {
 }
 
diff --git a/Engine/Source/Editor/UnrealEd/Private/Commandlets/DumpHiddenCategoriesCommandlet.cpp b/Engine/Source/Editor/UnrealEd/Private/Commandlets/DumpHiddenCategoriesCommandlet.cpp
--- a/Engine/Source/Editor/UnrealEd/Private/Commandlets/DumpHiddenCategoriesCommandlet.cpp
+++ b/Engine/Source/Editor/UnrealEd/Private/Commandlets/DumpHiddenCategoriesCommandlet.cpp
@@ -162,8 +162,8 @@ static int32 GetHiddenProperties(uint32 Indent, UClass* Class, FString& JsonOut)
 *******************************************************************************/
 
 //------------------------------------------------------------------------------
-UDumpHiddenCategoriesCommandlet::UDumpHiddenCategoriesCommandlet(class FPostConstructInitializeProperties const& PCIP)
-	: Super(PCIP)
+UDumpHiddenCategoriesCommandlet::UDumpHiddenCategoriesCommandlet(FObjectInitializer const& ObjectInitializer)
+	: Super(ObjectInitializer)
 {
 }
 
diff --git a/Engine/Source/Editor/UnrealEd/Private/Commandlets/FileServerCommandlet.cpp b/Engine/Source/Editor/UnrealEd/Private/Commandlets/FileServerCommandlet.cpp
--- a/Engine/Source/Editor/UnrealEd/Private/Commandlets/FileServerCommandlet.cpp
+++ b/Engine/Source/Editor/UnrealEd/Private/Commandlets/FileServerCommandlet.cpp
@@ -18,8 +18,8 @@ DEFINE_LOG_CATEGORY_STATIC(LogFileServerCommandlet, Log, All);
 /* UFileServerCommandlet structors
  *****************************************************************************/
 
-UFileServerCommandlet::UFileServerCommandlet( const class FPostConstructInitializeProperties& PCIP )
-	: Super(PCIP)
+UFileServerCommandlet::UFileServerCommandlet( const FObjectInitializer& ObjectInitializer )
+	: Super(ObjectInitializer)
 {
 
 	IsClient = false;
diff --git a/Engine/Source/Editor/UnrealEd/Private/Commandlets/FixupRedirectsCommandlet.cpp b/Engine/Source/Editor/UnrealEd/Private/Commandlets/FixupRedirectsCommandlet.cpp
--- a/Engine/Source/Editor/UnrealEd/Private/Commandlets/FixupRedirectsCommandlet.cpp
+++ b/Engine/Source/Editor/UnrealEd/Private/Commandlets/FixupRedirectsCommandlet.cpp
@@ -9,8 +9,8 @@
 
 DEFINE_LOG_CATEGORY_STATIC(LogFixupRedirectsCommandlet, Log, All);
 
-UFixupRedirectsCommandlet::UFixupRedirectsCommandlet(const class FPostConstructInitializeProperties& PCIP)
-	: Super(PCIP)
+UFixupRedirectsCommandlet::UFixupRedirectsCommandlet(const FObjectInitializer& ObjectInitializer)
+	: Super(ObjectInitializer)
 {
 	LogToConsole = false;
 }
diff --git a/Engine/Source/Editor/UnrealEd/Private/Commandlets/GatherTextCommandlet.cpp b/Engine/Source/Editor/UnrealEd/Private/Commandlets/GatherTextCommandlet.cpp
--- a/Engine/Source/Editor/UnrealEd/Private/Commandlets/GatherTextCommandlet.cpp
+++ b/Engine/Source/Editor/UnrealEd/Private/Commandlets/GatherTextCommandlet.cpp
@@ -1,4 +1,4 @@
-﻿// Copyright 1998-2014 Epic Games, Inc. All Rights Reserved.
+// Copyright 1998-2014 Epic Games, Inc. All Rights Reserved.
 
 #include "UnrealEd.h"
 #include "ISourceControlModule.h"
@@ -9,8 +9,8 @@ DEFINE_LOG_CATEGORY_STATIC(LogGatherTextCommandlet, Log, All);
 /**
  *	UGatherTextCommandlet
  */
-UGatherTextCommandlet::UGatherTextCommandlet(const class FPostConstructInitializeProperties& PCIP)
-	: Super(PCIP)
+UGatherTextCommandlet::UGatherTextCommandlet(const FObjectInitializer& ObjectInitializer)
+	: Super(ObjectInitializer)
 {
 }
 
diff --git a/Engine/Source/Editor/UnrealEd/Private/Commandlets/GatherTextCommandletBase.cpp b/Engine/Source/Editor/UnrealEd/Private/Commandlets/GatherTextCommandletBase.cpp
--- a/Engine/Source/Editor/UnrealEd/Private/Commandlets/GatherTextCommandletBase.cpp
+++ b/Engine/Source/Editor/UnrealEd/Private/Commandlets/GatherTextCommandletBase.cpp
@@ -15,8 +15,8 @@ TSharedPtr<FConflictReportInfo> FConflictReportInfo::StaticConflictInstance;
 //////////////////////////////////////////////////////////////////////////
 //UGatherTextCommandletBase
 
-UGatherTextCommandletBase::UGatherTextCommandletBase(const class FPostConstructInitializeProperties& PCIP)
-	: Super(PCIP)
+UGatherTextCommandletBase::UGatherTextCommandletBase(const FObjectInitializer& ObjectInitializer)
+	: Super(ObjectInitializer)
 {
 }
 
diff --git a/Engine/Source/Editor/UnrealEd/Private/Commandlets/GatherTextFromAssetsCommandlet.cpp b/Engine/Source/Editor/UnrealEd/Private/Commandlets/GatherTextFromAssetsCommandlet.cpp
--- a/Engine/Source/Editor/UnrealEd/Private/Commandlets/GatherTextFromAssetsCommandlet.cpp
+++ b/Engine/Source/Editor/UnrealEd/Private/Commandlets/GatherTextFromAssetsCommandlet.cpp
@@ -1,4 +1,4 @@
-﻿// Copyright 1998-2014 Epic Games, Inc. All Rights Reserved.
+// Copyright 1998-2014 Epic Games, Inc. All Rights Reserved.
 
 #include "UnrealEd.h"
 #include "PackageTools.h"
@@ -341,8 +341,8 @@ TSharedPtr< FLocMetadataValue > UGatherTextFromAssetsCommandlet::FDialogueHelper
 }
 
 
-UGatherTextFromAssetsCommandlet::UGatherTextFromAssetsCommandlet(const class FPostConstructInitializeProperties& PCIP)
-	:Super(PCIP)
+UGatherTextFromAssetsCommandlet::UGatherTextFromAssetsCommandlet(const FObjectInitializer& ObjectInitializer)
+	:Super(ObjectInitializer)
 {
 
 }
diff --git a/Engine/Source/Editor/UnrealEd/Private/Commandlets/GatherTextFromMetadataCommandlet.cpp b/Engine/Source/Editor/UnrealEd/Private/Commandlets/GatherTextFromMetadataCommandlet.cpp
--- a/Engine/Source/Editor/UnrealEd/Private/Commandlets/GatherTextFromMetadataCommandlet.cpp
+++ b/Engine/Source/Editor/UnrealEd/Private/Commandlets/GatherTextFromMetadataCommandlet.cpp
@@ -8,8 +8,8 @@ DEFINE_LOG_CATEGORY_STATIC(LogGatherTextFromMetaDataCommandlet, Log, All);
 //////////////////////////////////////////////////////////////////////////
 //GatherTextFromMetaDataCommandlet
 
-UGatherTextFromMetaDataCommandlet::UGatherTextFromMetaDataCommandlet(const class FPostConstructInitializeProperties& PCIP)
-	: Super(PCIP)
+UGatherTextFromMetaDataCommandlet::UGatherTextFromMetaDataCommandlet(const FObjectInitializer& ObjectInitializer)
+	: Super(ObjectInitializer)
 {
 }
 
diff --git a/Engine/Source/Editor/UnrealEd/Private/Commandlets/GatherTextFromSourceCommandlet.cpp b/Engine/Source/Editor/UnrealEd/Private/Commandlets/GatherTextFromSourceCommandlet.cpp
--- a/Engine/Source/Editor/UnrealEd/Private/Commandlets/GatherTextFromSourceCommandlet.cpp
+++ b/Engine/Source/Editor/UnrealEd/Private/Commandlets/GatherTextFromSourceCommandlet.cpp
@@ -9,8 +9,8 @@ DEFINE_LOG_CATEGORY_STATIC(LogGatherTextFromSourceCommandlet, Log, All);
 
 #define LOC_DEFINE_REGION
 
-UGatherTextFromSourceCommandlet::UGatherTextFromSourceCommandlet(const class FPostConstructInitializeProperties& PCIP)
-	: Super(PCIP)
+UGatherTextFromSourceCommandlet::UGatherTextFromSourceCommandlet(const FObjectInitializer& ObjectInitializer)
+	: Super(ObjectInitializer)
 {
 }
 
diff --git a/Engine/Source/Editor/UnrealEd/Private/Commandlets/GenerateBlueprintAPICommandlet.cpp b/Engine/Source/Editor/UnrealEd/Private/Commandlets/GenerateBlueprintAPICommandlet.cpp
--- a/Engine/Source/Editor/UnrealEd/Private/Commandlets/GenerateBlueprintAPICommandlet.cpp
+++ b/Engine/Source/Editor/UnrealEd/Private/Commandlets/GenerateBlueprintAPICommandlet.cpp
@@ -937,8 +937,8 @@ static void GenerateBlueprintAPIUtils::DumpActionMenuItem(uint32 Indent, FGraphA
  ******************************************************************************/
 
 //------------------------------------------------------------------------------
-UGenerateBlueprintAPICommandlet::UGenerateBlueprintAPICommandlet(class FPostConstructInitializeProperties const& PCIP)
-	: Super(PCIP)
+UGenerateBlueprintAPICommandlet::UGenerateBlueprintAPICommandlet(FObjectInitializer const& ObjectInitializer)
+	: Super(ObjectInitializer)
 {
 }
 
diff --git a/Engine/Source/Editor/UnrealEd/Private/Commandlets/GenerateDistillFileSetsCommandlet.cpp b/Engine/Source/Editor/UnrealEd/Private/Commandlets/GenerateDistillFileSetsCommandlet.cpp
--- a/Engine/Source/Editor/UnrealEd/Private/Commandlets/GenerateDistillFileSetsCommandlet.cpp
+++ b/Engine/Source/Editor/UnrealEd/Private/Commandlets/GenerateDistillFileSetsCommandlet.cpp
@@ -5,8 +5,8 @@
 
 DEFINE_LOG_CATEGORY_STATIC(LogGenerateDistillFileSetsCommandlet, Log, All);
 
-UGenerateDistillFileSetsCommandlet::UGenerateDistillFileSetsCommandlet(const class FPostConstructInitializeProperties& PCIP)
-	: Super(PCIP)
+UGenerateDistillFileSetsCommandlet::UGenerateDistillFileSetsCommandlet(const FObjectInitializer& ObjectInitializer)
+	: Super(ObjectInitializer)
 {
 	
 }
diff --git a/Engine/Source/Editor/UnrealEd/Private/Commandlets/GenerateGatherArchiveCommandlet.cpp b/Engine/Source/Editor/UnrealEd/Private/Commandlets/GenerateGatherArchiveCommandlet.cpp
--- a/Engine/Source/Editor/UnrealEd/Private/Commandlets/GenerateGatherArchiveCommandlet.cpp
+++ b/Engine/Source/Editor/UnrealEd/Private/Commandlets/GenerateGatherArchiveCommandlet.cpp
@@ -15,8 +15,8 @@ DEFINE_LOG_CATEGORY_STATIC(LogGenerateArchiveCommandlet, Log, All);
 /**
  *	UGenerateGatherArchiveCommandlet
  */
-UGenerateGatherArchiveCommandlet::UGenerateGatherArchiveCommandlet(const class FPostConstructInitializeProperties& PCIP)
-	: Super(PCIP)
+UGenerateGatherArchiveCommandlet::UGenerateGatherArchiveCommandlet(const FObjectInitializer& ObjectInitializer)
+	: Super(ObjectInitializer)
 {
 }
 
diff --git a/Engine/Source/Editor/UnrealEd/Private/Commandlets/GenerateGatherManifestCommandlet.cpp b/Engine/Source/Editor/UnrealEd/Private/Commandlets/GenerateGatherManifestCommandlet.cpp
--- a/Engine/Source/Editor/UnrealEd/Private/Commandlets/GenerateGatherManifestCommandlet.cpp
+++ b/Engine/Source/Editor/UnrealEd/Private/Commandlets/GenerateGatherManifestCommandlet.cpp
@@ -11,8 +11,8 @@ DEFINE_LOG_CATEGORY_STATIC(LogGenerateManifestCommandlet, Log, All);
 /**
  *	UGenerateGatherManifestCommandlet
  */
-UGenerateGatherManifestCommandlet::UGenerateGatherManifestCommandlet(const class FPostConstructInitializeProperties& PCIP)
-	: Super(PCIP)
+UGenerateGatherManifestCommandlet::UGenerateGatherManifestCommandlet(const FObjectInitializer& ObjectInitializer)
+	: Super(ObjectInitializer)
 {
 }
 
diff --git a/Engine/Source/Editor/UnrealEd/Private/Commandlets/GenerateTextLocalizationReportCommandlet.cpp b/Engine/Source/Editor/UnrealEd/Private/Commandlets/GenerateTextLocalizationReportCommandlet.cpp
--- a/Engine/Source/Editor/UnrealEd/Private/Commandlets/GenerateTextLocalizationReportCommandlet.cpp
+++ b/Engine/Source/Editor/UnrealEd/Private/Commandlets/GenerateTextLocalizationReportCommandlet.cpp
@@ -10,8 +10,8 @@
 
 DEFINE_LOG_CATEGORY_STATIC(LogGenerateTextLocalizationReportCommandlet, Log, All);
 
-UGenerateTextLocalizationReportCommandlet::UGenerateTextLocalizationReportCommandlet(const class FPostConstructInitializeProperties& PCIP)
-	: Super(PCIP)
+UGenerateTextLocalizationReportCommandlet::UGenerateTextLocalizationReportCommandlet(const FObjectInitializer& ObjectInitializer)
+	: Super(ObjectInitializer)
 {
 }
 
diff --git a/Engine/Source/Editor/UnrealEd/Private/Commandlets/GenerateTextLocalizationResourceCommandlet.cpp b/Engine/Source/Editor/UnrealEd/Private/Commandlets/GenerateTextLocalizationResourceCommandlet.cpp
--- a/Engine/Source/Editor/UnrealEd/Private/Commandlets/GenerateTextLocalizationResourceCommandlet.cpp
+++ b/Engine/Source/Editor/UnrealEd/Private/Commandlets/GenerateTextLocalizationResourceCommandlet.cpp
@@ -1,4 +1,4 @@
-﻿// Copyright 1998-2014 Epic Games, Inc. All Rights Reserved.
+// Copyright 1998-2014 Epic Games, Inc. All Rights Reserved.
 
 #include "UnrealEd.h"
 #include "TextLocalizationResourceGenerator.h"
@@ -10,8 +10,8 @@
 
 DEFINE_LOG_CATEGORY_STATIC(LogGenerateTextLocalizationResourceCommandlet, Log, All);
 
-UGenerateTextLocalizationResourceCommandlet::UGenerateTextLocalizationResourceCommandlet(const class FPostConstructInitializeProperties& PCIP)
-	: Super(PCIP)
+UGenerateTextLocalizationResourceCommandlet::UGenerateTextLocalizationResourceCommandlet(const FObjectInitializer& ObjectInitializer)
+	: Super(ObjectInitializer)
 {
 }
 
diff --git a/Engine/Source/Editor/UnrealEd/Private/Commandlets/InternationalizationConditioningCommandlet.cpp b/Engine/Source/Editor/UnrealEd/Private/Commandlets/InternationalizationConditioningCommandlet.cpp
--- a/Engine/Source/Editor/UnrealEd/Private/Commandlets/InternationalizationConditioningCommandlet.cpp
+++ b/Engine/Source/Editor/UnrealEd/Private/Commandlets/InternationalizationConditioningCommandlet.cpp
@@ -205,8 +205,8 @@ bool UInternationalizationConditioningCommandlet::FLocalizationFilePair::HasFore
  *	UGatherTextCommandlet
  */
 
-UInternationalizationConditioningCommandlet::UInternationalizationConditioningCommandlet(const class FPostConstructInitializeProperties& PCIP)
-	: Super(PCIP)
+UInternationalizationConditioningCommandlet::UInternationalizationConditioningCommandlet(const FObjectInitializer& ObjectInitializer)
+	: Super(ObjectInitializer)
 {
 }
 
diff --git a/Engine/Source/Editor/UnrealEd/Private/Commandlets/InternationalizationExportCommandlet.cpp b/Engine/Source/Editor/UnrealEd/Private/Commandlets/InternationalizationExportCommandlet.cpp
--- a/Engine/Source/Editor/UnrealEd/Private/Commandlets/InternationalizationExportCommandlet.cpp
+++ b/Engine/Source/Editor/UnrealEd/Private/Commandlets/InternationalizationExportCommandlet.cpp
@@ -1029,8 +1029,8 @@ FString FPortableObjectEntry::ToString() const
 /**
  *	UInternationalizationExportCommandlet
  */
-UInternationalizationExportCommandlet::UInternationalizationExportCommandlet(const class FPostConstructInitializeProperties& PCIP)
-	: Super(PCIP)
+UInternationalizationExportCommandlet::UInternationalizationExportCommandlet(const FObjectInitializer& ObjectInitializer)
+	: Super(ObjectInitializer)
 {
 }
 
diff --git a/Engine/Source/Editor/UnrealEd/Private/Commandlets/KismetUpdateCommandlet.cpp b/Engine/Source/Editor/UnrealEd/Private/Commandlets/KismetUpdateCommandlet.cpp
--- a/Engine/Source/Editor/UnrealEd/Private/Commandlets/KismetUpdateCommandlet.cpp
+++ b/Engine/Source/Editor/UnrealEd/Private/Commandlets/KismetUpdateCommandlet.cpp
@@ -12,8 +12,8 @@ DEFINE_LOG_CATEGORY_STATIC(LogBlueprintUpdateCommandlet, Log, All);
 //////////////////////////////////////////////////////////////////////////
 // UKismetUpdateCommandlet
 
-UKismetUpdateCommandlet::UKismetUpdateCommandlet(const class FPostConstructInitializeProperties& PCIP)
-	: Super(PCIP)
+UKismetUpdateCommandlet::UKismetUpdateCommandlet(const FObjectInitializer& ObjectInitializer)
+	: Super(ObjectInitializer)
 {
 }
 
