Fixed vertex editing in geometry mode with rotated/scaled brushes.

diff --git a/Engine/Source/Editor/GeometryMode/Private/GeometryEdMode.cpp b/Engine/Source/Editor/GeometryMode/Private/GeometryEdMode.cpp
--- a/Engine/Source/Editor/GeometryMode/Private/GeometryEdMode.cpp
+++ b/Engine/Source/Editor/GeometryMode/Private/GeometryEdMode.cpp
@@ -80,7 +80,11 @@ bool FEdModeGeometry::GetCustomDrawingCoordinateSystem( FMatrix& InMatrix, void*
 
 	if( InData )
 	{
-		InMatrix = FRotationMatrix( ((FGeomBase*)InData)->GetNormal().Rotation() );
+		FGeomBase* GeomBase = static_cast<FGeomBase*>(InData);
+		FGeomObject* GeomObject = GeomBase->GetParentObject();
+		check(GeomObject != nullptr);
+		ABrush* Brush = GeomObject->GetActualBrush();
+		InMatrix = FRotationMatrix(GeomBase->GetNormal().Rotation()) * FRotationMatrix(Brush->GetActorRotation());
 	}
 	else
 	{
@@ -94,7 +98,12 @@ bool FEdModeGeometry::GetCustomDrawingCoordinateSystem( FMatrix& InMatrix, void*
 
 			if( go->SelectionOrder.Num() )
 			{
-				InMatrix = FRotationMatrix( go->SelectionOrder[ go->SelectionOrder.Num()-1 ]->GetWidgetRotation() );
+				FGeomBase* GeomBase = go->SelectionOrder[go->SelectionOrder.Num() - 1];
+				check(GeomBase != nullptr);
+				FGeomObject* GeomObject = GeomBase->GetParentObject();
+				check(GeomObject != nullptr);
+				ABrush* Brush = GeomObject->GetActualBrush();
+				InMatrix = FRotationMatrix( go->SelectionOrder[ go->SelectionOrder.Num()-1 ]->GetWidgetRotation() ) * FRotationMatrix(Brush->GetActorRotation());
 				return 1;
 			}
 		}
diff --git a/Engine/Source/Editor/GeometryMode/Private/GeometryModifiers.cpp b/Engine/Source/Editor/GeometryMode/Private/GeometryModifiers.cpp
--- a/Engine/Source/Editor/GeometryMode/Private/GeometryModifiers.cpp
+++ b/Engine/Source/Editor/GeometryMode/Private/GeometryModifiers.cpp
@@ -488,17 +488,8 @@ bool UGeomModifier_Edit::InputDelta(FEditorViewportClient* InViewportClient,FVie
 	
 	if( !InDrag.IsZero() )
 	{
-		// Volumes store their rotation locally, whereas normal brush rotations are always in worldspace, so we need to transform
-		// the drag vector into a volume's local space before applying it.
-		if( Brush->IsVolumeBrush() )
-		{
-			const FVector Wk = Brush->ActorToWorld().InverseTransformVector(InDrag);
-			vertOffset = Wk;
-		}
-		else
-		{
-			vertOffset = InDrag;
-		}
+		// Transform the drag vector into the brush's local space before applying it.
+		vertOffset = Brush->ActorToWorld().InverseTransformVector(InDrag);
 	}
 
 	/**
@@ -531,17 +522,19 @@ bool UGeomModifier_Edit::InputDelta(FEditorViewportClient* InViewportClient,FVie
 
 		if( !InScale.IsZero() )
 		{
-			float XFactor = (InScale.X > 0.f) ? 1 : -1;
-			float YFactor = (InScale.Y > 0.f) ? 1 : -1;
-			float ZFactor = (InScale.Z > 0.f) ? 1 : -1;
+			// TODO: sort this out properly, taking into account the 'fake' local transform according to the poly normal.
+			// Also use a scaling method which actually works properly.
+
+			float XFactor = FMath::Sign(InScale.X);
+			float YFactor = FMath::Sign(InScale.Y);
+			float ZFactor = FMath::Sign(InScale.Z);
 			float Strength;
 
 			FVector Wk( vtx->X, vtx->Y, vtx->Z );
-			Wk = vtx->GetParentObject()->GetActualBrush()->ActorToWorld().TransformPosition( Wk );
 
 			// Move vert to the origin
 
-			Wk -= GLevelEditorModeTools().PivotLocation;
+			Wk -= (GLevelEditorModeTools().PivotLocation - Brush->GetActorLocation());
 
 			// Move it along each axis based on it's distance from the origin
 
@@ -565,9 +558,9 @@ bool UGeomModifier_Edit::InputDelta(FEditorViewportClient* InViewportClient,FVie
 
 			// Move it back into world space
 
-			Wk += GLevelEditorModeTools().PivotLocation;
+			Wk += (GLevelEditorModeTools().PivotLocation - Brush->GetActorLocation());
 
-			*vtx = vtx->GetParentObject()->GetActualBrush()->ActorToWorld().InverseTransformPosition( Wk );
+			*vtx = Wk;
 		}
 	}
 	
