Fixed assert when building a brush with a spiral arrangement of vertices.

diff --git a/Engine/Source/Editor/GeometryMode/Private/GeometryModifiers.cpp b/Engine/Source/Editor/GeometryMode/Private/GeometryModifiers.cpp
--- a/Engine/Source/Editor/GeometryMode/Private/GeometryModifiers.cpp
+++ b/Engine/Source/Editor/GeometryMode/Private/GeometryModifiers.cpp
@@ -1397,6 +1397,136 @@ void UGeomModifier_Pen::Apply()
 	}
 }
 
+#if 0
+static bool DoesFinalLineIntersectWithShape(const TArray<FVector>& Vertices, const FVector& FinalVertex)
+{
+	// TODO: improve this, it often fails.
+
+	// Determine if the next line segment would intersect with any of the previous ones in the shape
+	for (int32 Index = 0; Index < Vertices.Num() - 1; Index++)
+	{
+		const FVector& Segment1Start = Vertices[Index];
+		const FVector& Segment1End = Vertices[Index + 1];
+		const FVector& Segment2Start = Vertices[Vertices.Num() - 1];
+		const FVector& Segment2End = FinalVertex;
+
+		// Check that the two line segments are coplanar
+		check(FMath::IsNearlyZero(FVector::DotProduct(Segment2Start - Segment1Start, FVector::CrossProduct(Segment1End - Segment1Start, Segment2End - Segment2Start))));
+
+		FVector Segment1Result;
+		FVector Segment2Result;
+		FMath::SegmentDistToSegmentSafe(Segment1Start, Segment1End, Segment2Start, Segment2End, Segment1Result, Segment2Result);
+		if (Segment1Result.Equals(Segment2Result) && !Segment1Result.Equals(Segment1Start) && !Segment1Result.Equals(Segment1End))
+		{
+			return true;
+		}
+	}
+
+	return false;
+}
+#endif
+
+static bool DoLineSegmentsIntersect(const FVector2D& Segment1Start, const FVector2D& Segment1End, const FVector2D& Segment2Start, const FVector2D& Segment2End)
+{
+	const FVector2D Segment1Dir = Segment1End - Segment1Start;
+	const FVector2D Segment2Dir = Segment2End - Segment2Start;
+
+	const float Determinant = FVector2D::CrossProduct(Segment1Dir, Segment2Dir);
+	if (!FMath::IsNearlyZero(Determinant))
+	{
+		const FVector2D SegmentStartDelta = Segment2Start - Segment1Start;
+		const float OneOverDet = 1.0f / Determinant;
+		const float Seg1Intersection = FVector2D::CrossProduct(SegmentStartDelta, Segment2Dir) * OneOverDet;
+		const float Seg2Intersection = FVector2D::CrossProduct(SegmentStartDelta, Segment1Dir) * OneOverDet;
+
+		const float Epsilon = 1/128.0f;
+		return (Seg1Intersection > Epsilon && Seg1Intersection < 1.0f - Epsilon && Seg2Intersection > Epsilon && Seg2Intersection < 1.0f - Epsilon);
+	}
+
+	return false;
+}
+
+/**
+* Given an array of points forming an unclosed polygon, determines whether a line segment formed from the final polygon vertex and a given endpoint
+* intersects with any edge of the polygon in the 2D plane in which both segments lie.
+*
+* @param	Vertices		Array of vertices forming unclosed polygon
+* @param	EndVertex		Endpoint of line segment starting from final point of polygon
+*/
+static bool DoesFinalLineIntersectWithShape(const TArray<FVector>& Vertices, const FVector& EndVertex)
+{
+	// Can't intersect if there are fewer than 2 vertices
+	if (Vertices.Num() < 2)
+	{
+		return false;
+	}
+
+	// All line segments in the polygon ought to be coplanar. Hence the problem can be reduced to detecting intersections of line segments
+	// projected onto their common plane, using 2D coordinates.
+
+	// Line segment 1 is the line to test against the rest of the shape
+	const FVector& Segment1Start = Vertices[Vertices.Num() - 1];
+	const FVector& Segment1End = EndVertex;
+
+	const FVector Segment1Dir = Segment1End - Segment1Start;
+	const float Segment1Len = Segment1Dir.Size();
+	if (FMath::IsNearlyZero(Segment1Len))
+	{
+		// Treat zero length line segments as non-intersecting
+		return false;
+	}
+
+	// The direction of segment 1 on the plane will provide the X axis of the 2D basis on the plane
+	const FVector ProjectedXAxis = Segment1Dir / Segment1Len;
+
+	for (int32 Index = 0; Index < Vertices.Num() - 1; Index++)
+	{
+		// Line segment 2 is each edge of the shape
+		const FVector& Segment2Start = Vertices[Index];
+		const FVector& Segment2End = Vertices[Index + 1];
+		const FVector Segment2Dir = Segment2End - Segment2Start;
+
+		const FVector SegmentStartDelta = Segment2Start - Segment1Start;
+
+		// If line segments 1 and 2 are coplanar, the plane normal will be shared, and be calculated from a cross product of their two directions
+		FVector PlaneNormal = FVector::CrossProduct(Segment1Dir, Segment2Dir);
+
+		// Check that they are indeed coplanar
+		const bool bIsCoplanar = FMath::IsNearlyZero(FVector::DotProduct(SegmentStartDelta, PlaneNormal));
+		if (!bIsCoplanar)
+		{
+			// Non-coplanar line segments can't possibly intersect (ignoring the case of coincident start/end points)
+			return false;
+		}
+
+		// Parallel line segments will have yielded a zero normal. Attempt to calculate it from the segment start delta.
+		// If the lines are coincident, this will also yield a zero normal, resulting in a 1D basis (which is still sufficient to detect segment overlaps in projection space).
+		const bool bParallel = (FMath::IsNearlyZero(PlaneNormal.SizeSquared()));
+		if (bParallel)
+		{
+			PlaneNormal = FVector::CrossProduct(Segment1Dir, SegmentStartDelta);
+		}
+
+		// Get the Y axis of the 2D basis from the X axis and the plane normal
+		const FVector ProjectedYAxis = FVector::CrossProduct(PlaneNormal.GetSafeNormal(), ProjectedXAxis);
+
+		// Project 3d points onto plane
+		const FVector2D ProjectedSegment1Start(0.0f, 0.0f);
+		const FVector2D ProjectedSegment1End(Segment1Len, 0.0f);
+		const FVector2D ProjectedSegment2Start(FVector::DotProduct(ProjectedXAxis, SegmentStartDelta), FVector::DotProduct(ProjectedYAxis, SegmentStartDelta));
+		const FVector2D ProjectedSegment2End(FVector::DotProduct(ProjectedXAxis, Segment2End - Segment1Start), FVector::DotProduct(ProjectedYAxis, Segment2End - Segment1Start));
+
+		// Now check intersection of 2d segments.
+		if (DoLineSegmentsIntersect(ProjectedSegment1Start, ProjectedSegment1End, ProjectedSegment2Start, ProjectedSegment2End))
+		{
+			return true;
+		}
+	}
+
+	return false;
+}
+
+
 /**
 * @return		true if the key was handled by this editor mode tool.
 */
@@ -1422,15 +1552,21 @@ bool UGeomModifier_Pen::InputKey(FEditorViewportClient* ViewportClient, FViewpor
 			}
 			if( ShapeVertices.Num() && MouseWorldSpacePos.Equals( ShapeVertices[0] ) )
 			{
-				Apply();
+				if (!DoesFinalLineIntersectWithShape(ShapeVertices, ShapeVertices[0]))
+				{
+					Apply();
+					bResult = true;
+				}
 			}
 			else
 			{
-				UsingViewportClient = ViewportClient;
-				ShapeVertices.Add( MouseWorldSpacePos );
+				if (!DoesFinalLineIntersectWithShape(ShapeVertices, MouseWorldSpacePos))
+				{
+					UsingViewportClient = ViewportClient;
+					ShapeVertices.Add( MouseWorldSpacePos );
+					bResult = true;
+				}
 			}
-
-			bResult = true;
 		}
 		else if( Key == EKeys::Escape || Key == EKeys::BackSpace )
 		{
@@ -1443,9 +1579,11 @@ bool UGeomModifier_Pen::InputKey(FEditorViewportClient* ViewportClient, FViewpor
 		}
 		else if( Key == EKeys::Enter )
 		{
-			Apply();
-
-			bResult = true;
+			if (!DoesFinalLineIntersectWithShape(ShapeVertices, ShapeVertices[0]))
+			{
+				Apply();
+				bResult = true;
+			}
 		}
 	}
 
@@ -1495,14 +1633,20 @@ void UGeomModifier_Pen::Render(const FSceneView* View,FViewport* Viewport,FPrimi
 
 	if( ShapeVertices.Num() )
 	{
-		// Draw a dashed line from the last placed vertex to the current mouse position
-		DrawDashedLine( PDI, ShapeVertices[ ShapeVertices.Num() - 1 ], MouseWorldSpacePos, FLinearColor(1,0.5f,0), GEditor->GetGridSize(), SDPG_Foreground );
+		if (!DoesFinalLineIntersectWithShape(ShapeVertices, MouseWorldSpacePos))
+		{
+			// Draw a dashed line from the last placed vertex to the current mouse position
+			DrawDashedLine(PDI, ShapeVertices[ShapeVertices.Num() - 1], MouseWorldSpacePos, FLinearColor(1, 0.5f, 0), GEditor->GetGridSize(), SDPG_Foreground);
+		}
 	}
 
 	if( ShapeVertices.Num() > 2 )
 	{
-		// Draw a darkened dashed line to show what the completed shape will look like
-		DrawDashedLine( PDI, ShapeVertices[ ShapeVertices.Num() - 1 ], ShapeVertices[ 0 ], FLinearColor(.5,0,0), GEditor->GetGridSize(), SDPG_Foreground );
+		if (!DoesFinalLineIntersectWithShape(ShapeVertices, ShapeVertices[0]))
+		{
+			// Draw a darkened dashed line to show what the completed shape will look like
+			DrawDashedLine(PDI, ShapeVertices[ShapeVertices.Num() - 1], ShapeVertices[0], FLinearColor(.5, 0, 0), GEditor->GetGridSize(), SDPG_Foreground);
+		}
 	}
 
 	// Draw a box where the next vertex will be placed
diff --git a/Engine/Source/Editor/UnrealEd/Private/EditorBsp.cpp b/Engine/Source/Editor/UnrealEd/Private/EditorBsp.cpp
--- a/Engine/Source/Editor/UnrealEd/Private/EditorBsp.cpp
+++ b/Engine/Source/Editor/UnrealEd/Private/EditorBsp.cpp
@@ -1210,6 +1210,7 @@ int UEditorEngine::bspBrushCSG
 	else
 	{
 		// Add and subtract.
+		TMap<int32, int32> SurfaceIndexRemap;
 		for( i=0; i<TempModel->Polys->Element.Num(); i++ )
 		{
          	FPoly EdPoly = TempModel->Polys->Element[i];
@@ -1218,13 +1219,16 @@ int UEditorEngine::bspBrushCSG
          	// be split, and set iLink so that BspAddNode will know to add its information
          	// if a node is added based on this poly.
          	EdPoly.PolyFlags &= ~(PF_EdCut);
-          	if( EdPoly.iLink == i )
+			const int32* SurfaceIndexPtr = SurfaceIndexRemap.Find(EdPoly.iLink);
+			if (SurfaceIndexPtr == nullptr)
 			{
-				EdPoly.iLink = TempModel->Polys->Element[i].iLink = Model->Surfs.Num();
+				const int32 NewSurfaceIndex = Model->Surfs.Num();
+				SurfaceIndexRemap.Add(EdPoly.iLink, NewSurfaceIndex);
+				EdPoly.iLink = TempModel->Polys->Element[i].iLink = NewSurfaceIndex;
 			}
 			else
 			{
-				EdPoly.iLink = TempModel->Polys->Element[EdPoly.iLink].iLink;
+				EdPoly.iLink = TempModel->Polys->Element[i].iLink = *SurfaceIndexPtr;
 			}
 
 			// Filter brush through the world.
