Fixed crash when attempting to turn an edge in Geometry Mode whose associated polys are coincident.

diff --git a/Engine/Source/Editor/GeometryMode/Private/GeometryModifiers.cpp b/Engine/Source/Editor/GeometryMode/Private/GeometryModifiers.cpp
--- a/Engine/Source/Editor/GeometryMode/Private/GeometryModifiers.cpp
+++ b/Engine/Source/Editor/GeometryMode/Private/GeometryModifiers.cpp
@@ -3155,6 +3155,12 @@ bool UGeomModifier_Turn::OnApply()
 					Quad.AddUnique( Poly->Vertices[v] );
 				}
 
+				// If the adjoining polys were coincident, don't try to turn the edge
+				if (Quad.Num() == 3)
+				{
+					continue;
+				}
+
 				// Create new polygons
 
 				FPoly* NewPoly;
