Fixed edit/extrude BSP operation with multiply selected brushes. Also fixed scale for the limited case that there is exactly one surface selected.

diff --git a/Engine/Source/Editor/GeometryMode/Private/GeometryModifiers.cpp b/Engine/Source/Editor/GeometryMode/Private/GeometryModifiers.cpp
--- a/Engine/Source/Editor/GeometryMode/Private/GeometryModifiers.cpp
+++ b/Engine/Source/Editor/GeometryMode/Private/GeometryModifiers.cpp
@@ -401,6 +401,10 @@ bool UGeomModifier_Edit::InputDelta(FEditorViewportClient* InViewportClient,FVie
 	* vertices a little differently.
 	*/
 
+	// Only permit scaling if there is exactly one selected poly. Make a note of it here
+	FGeomPoly* SelectedPoly = nullptr;
+	int32 NumSelectedPolys = 0;
+
 	for( FEdModeGeometry::TGeomObjectIterator Itor( mode->GeomObjectItor() ) ; Itor ; ++Itor )
 	{
 		FGeomObject* go = *Itor;
@@ -410,6 +414,9 @@ bool UGeomModifier_Edit::InputDelta(FEditorViewportClient* InViewportClient,FVie
 			FGeomPoly* gp = &go->PolyPool[p];
 			if( gp->IsSelected() )
 			{
+				SelectedPoly = gp;
+				NumSelectedPolys++;
+
 				for( int32 e = 0 ; e < gp->EdgeIndices.Num() ; ++e )
 				{
 					FGeomEdge* ge = &go->EdgePool[ gp->EdgeIndices[e] ];
@@ -480,86 +487,55 @@ bool UGeomModifier_Edit::InputDelta(FEditorViewportClient* InViewportClient,FVie
 
 	FVector BBoxExtent = VertBBox.GetExtent();
 
-	FGeomVertex* vertex0 = UniqueVertexList[0];
-	const ABrush* Brush = vertex0->GetParentObject()->GetActualBrush();
-	FVector vertOffset = FVector(0, 0, 0);
-
-	//Calculate translation now so that it isn't done every iteration of the proceeding loop
-	
-	if( !InDrag.IsZero() )
-	{
-		// Transform the drag vector into the brush's local space before applying it.
-		vertOffset = Brush->ActorToWorld().InverseTransformVector(InDrag);
-	}
-
 	/**
-	* We first generate a list of unique vertices and then transform that list
-	* in one shot.  This prevents vertices from being touched more than once (which
-	* would result in them transforming x times as fast as others).
-	*/
+	 * We first generate a list of unique vertices and then transform that list
+	 * in one shot.  This prevents vertices from being touched more than once (which
+	 * would result in them transforming x times as fast as others).
+	 */
 	for( int32 x = 0 ; x < UniqueVertexList.Num() ; ++x )
 	{
 		FGeomVertex* vtx = UniqueVertexList[x];
-		//Translate
+		const ABrush* Brush = vtx->GetParentObject()->GetActualBrush();
+
+		// Translate
 
-			*vtx += vertOffset;
+		if (!InDrag.IsZero())
+		{
+			*vtx += Brush->ActorToWorld().InverseTransformVector(InDrag);
+		}
 
 		// Rotate
 
 		if( !FinalRot.IsZero() )
 		{
-			FRotationMatrix matrix( FinalRot );
+			const FRotationMatrix Matrix( FinalRot );
 
 			FVector Wk( vtx->X, vtx->Y, vtx->Z );
-			Wk = vtx->GetParentObject()->GetActualBrush()->ActorToWorld().TransformPosition( Wk );
+			Wk = Brush->ActorToWorld().TransformPosition(Wk);
 			Wk -= GLevelEditorModeTools().PivotLocation;
-			Wk = matrix.TransformPosition( Wk );
+			Wk = Matrix.TransformPosition( Wk );
 			Wk += GLevelEditorModeTools().PivotLocation;
-			*vtx = vtx->GetParentObject()->GetActualBrush()->ActorToWorld().InverseTransformPosition( Wk );
+			*vtx = Brush->ActorToWorld().InverseTransformPosition(Wk);
 		}
 
 		// Scale
 
-		if( !InScale.IsZero() )
-		{
-			// TODO: sort this out properly, taking into account the 'fake' local transform according to the poly normal.
-			// Also use a scaling method which actually works properly.
-
-			float XFactor = FMath::Sign(InScale.X);
-			float YFactor = FMath::Sign(InScale.Y);
-			float ZFactor = FMath::Sign(InScale.Z);
-			float Strength;
-
-			FVector Wk( vtx->X, vtx->Y, vtx->Z );
-
-			// Move vert to the origin
-
-			Wk -= (GLevelEditorModeTools().PivotLocation - Brush->GetActorLocation());
-
-			// Move it along each axis based on it's distance from the origin
-
-			if( Wk.X != 0 )
-			{
-				Strength = (BBoxExtent.X / Wk.X) * XFactor;
-				Wk.X += GEditor->GetGridSize() * Strength;
-			}
-
-			if( Wk.Y != 0 )
-			{
-				Strength = (BBoxExtent.Y / Wk.Y) * YFactor;
-				Wk.Y += GEditor->GetGridSize() * Strength;
-			}
-
-			if( Wk.Z != 0 )
-			{
-				Strength = (BBoxExtent.Z / Wk.Z) * ZFactor;
-				Wk.Z += GEditor->GetGridSize() * Strength;
-			}
-
-			// Move it back into world space
-
-			Wk += (GLevelEditorModeTools().PivotLocation - Brush->GetActorLocation());
-
+		if( !InScale.IsZero() && NumSelectedPolys == 1)
+		{
+			// This is a quick fix for now.
+			// Scale needs to know the surface normal of the polys selected (and scale only makes sense on polys, not edges or verts),
+			// so remember one selected poly and use that transform.
+			// Since scaling is relative to the pivot, it would actually be horrible scaling multiple polys at once anyway.
+			const FScaleMatrix Matrix(InScale + 1.0f);
+			const FVector PivotInModelSpace = Brush->ActorToWorld().InverseTransformPosition(GLevelEditorModeTools().PivotLocation);
+			const FRotationMatrix GeomBaseTransform = FRotationMatrix(SelectedPoly->GetWidgetRotation());
+
+			FVector Wk(vtx->X, vtx->Y, vtx->Z);
+			Wk -= PivotInModelSpace;
+			Wk = GeomBaseTransform.TransformPosition(Wk);
+			Wk = Matrix.TransformPosition(Wk);
+			Wk = GeomBaseTransform.InverseTransformPosition(Wk);
+			Wk += PivotInModelSpace;
 			*vtx = Wk;
 		}
 	}
@@ -570,11 +546,12 @@ bool UGeomModifier_Edit::InputDelta(FEditorViewportClient* InViewportClient,FVie
 		for( int32 x = 0 ; x < UniqueVertexList.Num() ; ++x )
 		{
 			FGeomVertex* vtx = UniqueVertexList[x];
-			*vtx -= vertOffset;
+			const ABrush* Brush = vtx->GetParentObject()->GetActualBrush();
+			*vtx -= Brush->ActorToWorld().InverseTransformVector(InDrag);
 		}
 		
-		GLevelEditorModeTools().PivotLocation -= vertOffset;
-		GLevelEditorModeTools().SnappedLocation -= vertOffset;
+		GLevelEditorModeTools().PivotLocation -= InDrag;
+		GLevelEditorModeTools().SnappedLocation -= InDrag;
 	}
 
 	const bool bIsCtrlPressed = InViewportClient->IsCtrlPressed();
