Added ToolTip text to the various geometry modifier widgets.

diff --git a/Engine/Source/Editor/GeometryMode/Classes/GeomModifier.h b/Engine/Source/Editor/GeometryMode/Classes/GeomModifier.h
--- a/Engine/Source/Editor/GeometryMode/Classes/GeomModifier.h
+++ b/Engine/Source/Editor/GeometryMode/Classes/GeomModifier.h
@@ -17,6 +17,10 @@ class UGeomModifier : public UObject
 	UPROPERTY(EditAnywhere, Category=GeomModifier)
 	FText Description;
 
+	/** The tooltip to be displayed for this modifier */
+	UPROPERTY(EditAnywhere, Category = GeomModifier)
+	FText Tooltip;
+
 	/** If true, this modifier should be displayed as a push button instead of a radio button */
 	UPROPERTY(EditAnywhere, Category=GeomModifier)
 	uint32 bPushButton:1;
@@ -38,6 +42,9 @@ public:
 	/** @return		The modifier's description string. */
 	const FText& GetModifierDescription() const;
 
+	/** @return		The modifier's tooltip string. */
+	const FText& GetModifierTooltip() const;
+
 	/** @return		true if the key was handled by this editor mode tool. */
 	virtual bool InputKey(class FEditorViewportClient* ViewportClient, FViewport* Viewport, FKey Key, EInputEvent Event);
 
diff --git a/Engine/Source/Editor/GeometryMode/Private/GeometryMode.cpp b/Engine/Source/Editor/GeometryMode/Private/GeometryMode.cpp
--- a/Engine/Source/Editor/GeometryMode/Private/GeometryMode.cpp
+++ b/Engine/Source/Editor/GeometryMode/Private/GeometryMode.cpp
@@ -238,6 +238,7 @@ TSharedRef<SCheckBox> SGeometryModeControls::CreateSingleModifierRadioButton(UGe
 	.IsChecked(this, &SGeometryModeControls::IsModifierChecked, Modifier)
 	.IsEnabled(this, &SGeometryModeControls::IsModifierEnabled, Modifier)
 	.OnCheckStateChanged(this, &SGeometryModeControls::OnModifierStateChanged, Modifier)
+	.ToolTip(SNew(SToolTip).Text(Modifier->GetModifierTooltip()))
 	[
 		SNew(STextBlock).Text( Modifier->GetModifierDescription() )
 	];
@@ -252,6 +253,7 @@ TSharedRef<SButton> SGeometryModeControls::CreateSingleModifierButton(UGeomModif
 	TSharedRef<SButton> Widget =
 	SNew(SButton)
 	.Text( Modifier->GetModifierDescription() )
+	.ToolTip(SNew(SToolTip).Text(Modifier->GetModifierTooltip()))
 	.HAlign(HAlign_Center)
 	.IsEnabled(this, &SGeometryModeControls::IsModifierEnabled, Modifier)
 	.OnClicked(this, &SGeometryModeControls::OnModifierClicked, Modifier);
diff --git a/Engine/Source/Editor/GeometryMode/Private/GeometryModifiers.cpp b/Engine/Source/Editor/GeometryMode/Private/GeometryModifiers.cpp
--- a/Engine/Source/Editor/GeometryMode/Private/GeometryModifiers.cpp
+++ b/Engine/Source/Editor/GeometryMode/Private/GeometryModifiers.cpp
@@ -45,6 +45,11 @@ const FText& UGeomModifier::GetModifierDescription() const
 	return Description;
 }
 
+const FText& UGeomModifier::GetModifierTooltip() const
+{
+	return Tooltip;
+}
+
 void UGeomModifier::Initialize()
 {
 }
@@ -375,6 +380,7 @@ UGeomModifier_Edit::UGeomModifier_Edit(const FObjectInitializer& ObjectInitializ
 	: Super(ObjectInitializer)
 {
 	Description = NSLOCTEXT("UnrealEd", "Edit", "Edit");
+	Tooltip = NSLOCTEXT("UnrealEd.GeomModifier_Edit", "Tooltip", "Translate, rotate or scale existing geometry.");
 }
 
 
@@ -585,6 +591,7 @@ UGeomModifier_Extrude::UGeomModifier_Extrude(const FObjectInitializer& ObjectIni
 	: Super(ObjectInitializer)
 {
 	Description = NSLOCTEXT("UnrealEd", "Extrude", "Extrude");
+	Tooltip = NSLOCTEXT("UnrealEd.GeomModifier_Extrude", "Tooltip", "Moves the selected geometry element forward, creating new geometry behind it if necessary.");
 	Length = 16;
 	Segments = 1;
 }
@@ -816,6 +823,7 @@ UGeomModifier_Lathe::UGeomModifier_Lathe(const FObjectInitializer& ObjectInitial
 	: Super(ObjectInitializer)
 {
 	Description = NSLOCTEXT("UnrealEd", "Lathe", "Lathe");
+	Tooltip = NSLOCTEXT("UnrealEd.GeomModifier_Lathe", "Tooltip", "Create new geometry by rotating the selected brush shape about the current pivot point.");
 	Axis = EAxis::Y;
 	TotalSegments = 16;
 	Segments = 4;
@@ -1145,6 +1153,7 @@ UGeomModifier_Pen::UGeomModifier_Pen(const FObjectInitializer& ObjectInitializer
 	: Super(ObjectInitializer)
 {
 	Description = NSLOCTEXT("UnrealEd", "Pen", "Pen");
+	Tooltip = NSLOCTEXT("UnrealEd.GeomModifier_Pen", "Tooltip", "Create new geometry by drawing the vertices directly into an orthographic viewport. Press space bar to place a vertex, and Enter to close the polygon.");
 	bCreateBrushShape = false;
 	bAutoExtrude = true;
 	ExtrudeDepth = 256;
@@ -1940,6 +1949,7 @@ UGeomModifier_Clip::UGeomModifier_Clip(const FObjectInitializer& ObjectInitializ
 	: Super(ObjectInitializer)
 {
 	Description = NSLOCTEXT("UnrealEd", "BrushClip", "Brush Clip");
+	Tooltip = NSLOCTEXT("UnrealEd.GeomModifier_Clip", "Tooltip", "Given a dividing plane, cut the geometry into two pieces, optionally discarding one of them. This operation only works in an orthographic viewport.  Define the vertices of the dividing plane with the space bar, and press Enter to apply.");
 	bFlipNormal = false;
 	bSplit = false;
 }
@@ -2320,6 +2330,7 @@ UGeomModifier_Delete::UGeomModifier_Delete(const FObjectInitializer& ObjectIniti
 	: Super(ObjectInitializer)
 {
 	Description = NSLOCTEXT("UnrealEd", "Delete", "Delete");
+	Tooltip = NSLOCTEXT("UnrealEd.GeomModifier_Delete", "Tooltip", "Deletes the selected geometry elements (vertices, edges or polygons).");
 	bPushButton = true;
 }
 
@@ -2405,6 +2416,7 @@ UGeomModifier_Create::UGeomModifier_Create(const FObjectInitializer& ObjectIniti
 	: Super(ObjectInitializer)
 {
 	Description = NSLOCTEXT("UnrealEd", "Create", "Create");
+	Tooltip = NSLOCTEXT("UnrealEd.GeomModifier_Create", "Tooltip", "Creates a new polygon from the selected vertices. The vertices must be selected in clockwise order to create a poly with an outward facing normal.");
 	bPushButton = true;
 }
 
@@ -2469,6 +2481,7 @@ UGeomModifier_Flip::UGeomModifier_Flip(const FObjectInitializer& ObjectInitializ
 	: Super(ObjectInitializer)
 {
 	Description = NSLOCTEXT("UnrealEd", "Flip", "Flip");
+	Tooltip = NSLOCTEXT("UnrealEd.GeomModifier_Flip", "Tooltip", "Flips the normal of the selected polygon so that it faces the other way.");
 	bPushButton = true;
 }
 
@@ -2513,6 +2526,7 @@ UGeomModifier_Split::UGeomModifier_Split(const FObjectInitializer& ObjectInitial
 	: Super(ObjectInitializer)
 {
 	Description = NSLOCTEXT("UnrealEd", "Split", "Split");
+	Tooltip = NSLOCTEXT("UnrealEd.GeomModifier_Pen", "Tooltip", "Split a brush in half, the exact operation depending on which geometry elements are selected.");
 	bPushButton = true;
 }
 
@@ -2883,6 +2897,7 @@ UGeomModifier_Triangulate::UGeomModifier_Triangulate(const FObjectInitializer& O
 	: Super(ObjectInitializer)
 {
 	Description = NSLOCTEXT("UnrealEd", "Triangulate", "Triangulate");
+	Tooltip = NSLOCTEXT("UnrealEd.GeomModifier_Triangulate", "Tooltip", "Break the selected polygons down into triangles.");
 	bPushButton = true;
 }
 
@@ -2895,7 +2910,7 @@ bool UGeomModifier_Triangulate::Supports()
 bool UGeomModifier_Triangulate::OnApply()
 {
 	FEdModeGeometry* mode = (FEdModeGeometry*)GLevelEditorModeTools().GetActiveMode(FBuiltinEditorModes::EM_Geometry);
-	bool bHavePolygonsSelected = mode->HavePolygonsSelected();
+	bool bHavePolygonsSelected = mode->HavePolygonsSelected(); 
 
 	// Mark the selected polygons so we can find them in the next loop, and create
 	// a local list of FPolys to triangulate later.
@@ -2954,6 +2969,7 @@ UGeomModifier_Optimize::UGeomModifier_Optimize(const FObjectInitializer& ObjectI
 	: Super(ObjectInitializer)
 {
 	Description = NSLOCTEXT("UnrealEd", "Optimize", "Optimize");
+	Tooltip = NSLOCTEXT("UnrealEd.GeomModifier_Optimize", "Tooltip", "Optimizes the selected geometry by merging together any polygons which can be formed into a single convex polygon.");
 	bPushButton = true;
 }
 
@@ -3041,6 +3057,7 @@ UGeomModifier_Turn::UGeomModifier_Turn(const FObjectInitializer& ObjectInitializ
 	: Super(ObjectInitializer)
 {
 	Description = NSLOCTEXT("UnrealEd", "Turn", "Turn");
+	Tooltip = NSLOCTEXT("UnrealEd.GeomModifier_Turn", "Tooltip", "Given a selected edge common to two triangles, turn the edge so that it is connected to the previously unconnected vertices.");
 	bPushButton = true;
 }
 
@@ -3207,6 +3224,7 @@ UGeomModifier_Weld::UGeomModifier_Weld(const FObjectInitializer& ObjectInitializ
 	: Super(ObjectInitializer)
 {
 	Description = NSLOCTEXT("UnrealEd", "Weld", "Weld");
+	Tooltip = NSLOCTEXT("UnrealEd.GeomModifier_Weld", "Tooltip", "Merge all selected vertices to the first selected vertex.");
 	bPushButton = true;
 }
 
