Fixed bugs in the Brush Clip tool.

The new UModel is created as an inner object of the new Brush, meaning that it has no external path and hence imports correctly.

diff --git a/Engine/Source/Editor/GeometryMode/Private/GeometryModifiers.cpp b/Engine/Source/Editor/GeometryMode/Private/GeometryModifiers.cpp
--- a/Engine/Source/Editor/GeometryMode/Private/GeometryModifiers.cpp
+++ b/Engine/Source/Editor/GeometryMode/Private/GeometryModifiers.cpp
@@ -1827,7 +1827,7 @@ static ABrush* ClipBrushAgainstPlane( const FPlane& InPlane, ABrush* InBrush)
 	// perhaps there were additional brushes were selected. 
 	check( ClippedBrush->GetClass() == InBrush->GetClass() );
 
-	ClippedBrush->Brush = NewObject<UModel>(InBrush->GetOuter());
+	ClippedBrush->Brush = NewObject<UModel>(ClippedBrush, NAME_None, RF_Transactional);
 	ClippedBrush->Brush->Initialize(nullptr);
 	ClippedBrush->GetBrushComponent()->Brush = ClippedBrush->Brush;
 
