Made BSP hit proxies use weak pointers to avoid stale data still in the hit proxy buffers when BSP geometry is recreated

diff --git a/Engine/Source/Editor/GeometryMode/Classes/GeomModifier.h b/Engine/Source/Editor/GeometryMode/Classes/GeomModifier.h
--- a/Engine/Source/Editor/GeometryMode/Classes/GeomModifier.h
+++ b/Engine/Source/Editor/GeometryMode/Classes/GeomModifier.h
@@ -7,6 +7,7 @@
  */
 
 #include "GeomModifier.generated.h"
+typedef TSharedPtr<class FGeomObject> FGeomObjectPtr;
 
 UCLASS(abstract, hidecategories=(Object, GeomModifier))
 class UGeomModifier : public UObject
@@ -95,7 +96,7 @@ public:
 	/**
 	 * Store the current geom selections (Edge, Vert and Poly)
 	 */
-	void StoreCurrentGeomSelections( TArray<struct FGeomSelection>& SelectionArray, FGeomObject* go );
+	void StoreCurrentGeomSelections( TArray<struct FGeomSelection>& SelectionArray, FGeomObjectPtr go );
 
 	/**
 	 * Store the current geom selections for all geom objects
diff --git a/Engine/Source/Editor/GeometryMode/Private/EditorGeometry.cpp b/Engine/Source/Editor/GeometryMode/Private/EditorGeometry.cpp
--- a/Engine/Source/Editor/GeometryMode/Private/EditorGeometry.cpp
+++ b/Engine/Source/Editor/GeometryMode/Private/EditorGeometry.cpp
@@ -43,7 +43,7 @@ void FGeomBase::Select( bool InSelect )
 	GetParentObject()->DirtySelectionOrder();
 }
 
-FGeomObject* FGeomBase::GetParentObject()
+FGeomObjectPtr FGeomBase::GetParentObject()
 {
 	check( GLevelEditorModeTools().IsModeActive(FBuiltinEditorModes::EM_Geometry) );
 	check( ParentObjectIndex > INDEX_NONE );
@@ -52,7 +52,7 @@ FGeomObject* FGeomBase::GetParentObject()
 	return mode->GetGeomObject( ParentObjectIndex );
 }
 
-const FGeomObject* FGeomBase::GetParentObject() const
+const FGeomObjectPtr FGeomBase::GetParentObject() const
 {
 	check( GLevelEditorModeTools().IsModeActive(FBuiltinEditorModes::EM_Geometry) );
 	check( ParentObjectIndex > INDEX_NONE );
@@ -221,8 +221,8 @@ FVector FGeomPoly::GetMidPoint() const
 
 FPoly* FGeomPoly::GetActualPoly()
 {
-	FGeomObject* Parent = GetParentObject();
-	check( Parent );
+	FGeomObjectPtr Parent = GetParentObject();
+	check( Parent.IsValid() );
 	ABrush* Brush = Parent->GetActualBrush();
 	check( Brush );
 
@@ -351,7 +351,7 @@ int32 FGeomObject::GetObjectIndex()
 
 	for( FEdModeGeometry::TGeomObjectIterator Itor( mode->GeomObjectItor() ) ; Itor ; ++Itor )
 	{
-		if( *Itor == this )
+		if( (*Itor).Get() == this )
 		{
 			return Itor.GetIndex();
 		}
diff --git a/Engine/Source/Editor/GeometryMode/Private/GeometryEdMode.cpp b/Engine/Source/Editor/GeometryMode/Private/GeometryEdMode.cpp
--- a/Engine/Source/Editor/GeometryMode/Private/GeometryEdMode.cpp
+++ b/Engine/Source/Editor/GeometryMode/Private/GeometryEdMode.cpp
@@ -36,11 +36,6 @@ FEdModeGeometry::FEdModeGeometry()
 
 FEdModeGeometry::~FEdModeGeometry()
 {
-	for( int32 i=0; i<GeomObjects.Num(); i++ )
-	{
-		FGeomObject* GeomObject	= GeomObjects[i];
-		delete GeomObject;
-	}
 	GeomObjects.Empty();
 }
 
@@ -81,8 +76,8 @@ bool FEdModeGeometry::GetCustomDrawingCoordinateSystem( FMatrix& InMatrix, void*
 	if( InData )
 	{
 		FGeomBase* GeomBase = static_cast<FGeomBase*>(InData);
-		FGeomObject* GeomObject = GeomBase->GetParentObject();
-		check(GeomObject != nullptr);
+		FGeomObjectPtr GeomObject = GeomBase->GetParentObject();
+		check(GeomObject.IsValid());
 		ABrush* Brush = GeomObject->GetActualBrush();
 		InMatrix = FRotationMatrix(GeomBase->GetNormal().Rotation()) * FQuatRotationMatrix(Brush->GetActorQuat());
 	}
@@ -93,15 +88,15 @@ bool FEdModeGeometry::GetCustomDrawingCoordinateSystem( FMatrix& InMatrix, void*
 
 		for( int32 o = 0 ; o < GeomObjects.Num() ; ++o )
 		{
-			FGeomObject* go = GeomObjects[o];
+			FGeomObjectPtr go = GeomObjects[o];
 			go->CompileSelectionOrder();
 
 			if( go->SelectionOrder.Num() )
 			{
 				FGeomBase* GeomBase = go->SelectionOrder[go->SelectionOrder.Num() - 1];
 				check(GeomBase != nullptr);
-				FGeomObject* GeomObject = GeomBase->GetParentObject();
-				check(GeomObject != nullptr);
+				FGeomObjectPtr GeomObject = GeomBase->GetParentObject();
+				check(GeomObject.IsValid());
 				ABrush* Brush = GeomObject->GetActualBrush();
 				InMatrix = FRotationMatrix( go->SelectionOrder[ go->SelectionOrder.Num()-1 ]->GetWidgetRotation() ) * FQuatRotationMatrix(Brush->GetActorQuat());
 				return 1;
@@ -142,11 +137,6 @@ void FEdModeGeometry::Exit()
 
 	FEdMode::Exit();
 
-	for( int32 i=0; i<GeomObjects.Num(); i++ )
-	{
-		FGeomObject* GeomObject	= GeomObjects[i];
-		delete GeomObject;
-	}
 	GeomObjects.Empty();
 }
 
@@ -192,7 +182,7 @@ int32 FEdModeGeometry::CountSelectedPolygons()
 
 	for( int32 ObjectIdx = 0 ; ObjectIdx < GeomObjects.Num() ; ++ObjectIdx )
 	{
-		FGeomObject* GeomObject = GeomObjects[ObjectIdx];
+		FGeomObjectPtr GeomObject = GeomObjects[ObjectIdx];
 
 		for( int32 P = 0 ; P < GeomObject->PolyPool.Num() ; ++P )
 		{
@@ -216,7 +206,7 @@ void FEdModeGeometry::GetSelectedPolygons( TArray<FGeomPoly*>& InPolygons )
 {
 	for( int32 ObjectIdx = 0 ; ObjectIdx < GeomObjects.Num() ; ++ObjectIdx )
 	{
-		FGeomObject* GeomObject = GeomObjects[ObjectIdx];
+		FGeomObjectPtr GeomObject = GeomObjects[ObjectIdx];
 
 		for( int32 P = 0 ; P < GeomObject->PolyPool.Num() ; ++P )
 		{
@@ -236,7 +226,7 @@ bool FEdModeGeometry::HavePolygonsSelected()
 {
 	for( int32 ObjectIdx = 0 ; ObjectIdx < GeomObjects.Num() ; ++ObjectIdx )
 	{
-		FGeomObject* GeomObject = GeomObjects[ObjectIdx];
+		FGeomObjectPtr GeomObject = GeomObjects[ObjectIdx];
 
 		for( int32 P = 0 ; P < GeomObject->PolyPool.Num() ; ++P )
 		{
@@ -260,7 +250,7 @@ int32 FEdModeGeometry::CountSelectedEdges()
 
 	for( int32 ObjectIdx = 0 ; ObjectIdx < GeomObjects.Num() ; ++ObjectIdx )
 	{
-		FGeomObject* GeomObject = GeomObjects[ObjectIdx];
+		FGeomObjectPtr GeomObject = GeomObjects[ObjectIdx];
 
 		for( int32 E = 0 ; E < GeomObject->EdgePool.Num() ; ++E )
 		{
@@ -282,7 +272,7 @@ bool FEdModeGeometry::HaveEdgesSelected()
 {
 	for( int32 ObjectIdx = 0 ; ObjectIdx < GeomObjects.Num() ; ++ObjectIdx )
 	{
-		FGeomObject* GeomObject = GeomObjects[ObjectIdx];
+		FGeomObjectPtr GeomObject = GeomObjects[ObjectIdx];
 
 		for( int32 E = 0 ; E < GeomObject->EdgePool.Num() ; ++E )
 		{
@@ -305,7 +295,7 @@ void FEdModeGeometry::GetSelectedEdges( TArray<FGeomEdge*>& InEdges )
 {
 	for( int32 ObjectIdx = 0 ; ObjectIdx < GeomObjects.Num() ; ++ObjectIdx )
 	{
-		FGeomObject* GeomObject = GeomObjects[ObjectIdx];
+		FGeomObjectPtr GeomObject = GeomObjects[ObjectIdx];
 
 		for( int32 E = 0 ; E < GeomObject->EdgePool.Num() ; ++E )
 		{
@@ -327,7 +317,7 @@ int32 FEdModeGeometry::CountSelectedVertices()
 
 	for( int32 ObjectIdx = 0 ; ObjectIdx < GeomObjects.Num() ; ++ObjectIdx )
 	{
-		FGeomObject* GeomObject = GeomObjects[ObjectIdx];
+		FGeomObjectPtr GeomObject = GeomObjects[ObjectIdx];
 
 		for( int32 V = 0 ; V < GeomObject->VertexPool.Num() ; ++V )
 		{
@@ -349,7 +339,7 @@ bool FEdModeGeometry::HaveVerticesSelected()
 {
 	for( int32 ObjectIdx = 0 ; ObjectIdx < GeomObjects.Num() ; ++ObjectIdx )
 	{
-		FGeomObject* GeomObject = GeomObjects[ObjectIdx];
+		FGeomObjectPtr GeomObject = GeomObjects[ObjectIdx];
 
 		for( int32 V = 0 ; V < GeomObject->VertexPool.Num() ; ++V )
 		{
@@ -374,7 +364,7 @@ void FEdModeGeometry::GetSelectedVertices( TArray<FGeomVertex*>& InVerts )
 
 	for( int32 ObjectIdx = 0 ; ObjectIdx < GeomObjects.Num() ; ++ObjectIdx )
 	{
-		FGeomObject* GeomObject = GeomObjects[ObjectIdx];
+		FGeomObjectPtr GeomObject = GeomObjects[ObjectIdx];
 
 		for( int32 V = 0 ; V < GeomObject->VertexPool.Num() ; ++V )
 		{
@@ -434,7 +424,7 @@ void FEdModeGeometry::GeometrySelectNone(bool bStoreSelection, bool bResetPivot)
 {
 	for( int32 ObjectIdx = 0 ; ObjectIdx < GeomObjects.Num() ; ++ObjectIdx )
 	{
-		FGeomObject* GeomObject = GeomObjects[ObjectIdx];
+		FGeomObjectPtr GeomObject = GeomObjects[ObjectIdx];
 		GeomObject->Select( 0 );
 
 		for( int VertexIdx = 0 ; VertexIdx < GeomObject->EdgePool.Num() ; ++VertexIdx )
@@ -476,7 +466,7 @@ void FEdModeGeometry::RenderPoly( const FSceneView* View, FViewport* Viewport, F
 {
 	for( int32 ObjectIdx = 0 ; ObjectIdx < GeomObjects.Num() ; ++ObjectIdx )
 	{
-		const FGeomObject* GeomObject = GeomObjects[ObjectIdx];
+		FGeomObjectPtr GeomObject = GeomObjects[ObjectIdx];
 		
 		FLinearColor UnselectedColor = GeomObject->GetActualBrush()->GetWireColor();
 		UnselectedColor.A = .1f;
@@ -495,7 +485,7 @@ void FEdModeGeometry::RenderPoly( const FSceneView* View, FViewport* Viewport, F
 		for( int32 PolyIdx = 0 ; PolyIdx < GeomObject->PolyPool.Num() ; ++PolyIdx )
 		{
 			const FGeomPoly* GeomPoly = &GeomObject->PolyPool[PolyIdx];
-			PDI->SetHitProxy( new HGeomPolyProxy(const_cast<FGeomObject*>(GeomPoly->GetParentObject()),PolyIdx) );
+			PDI->SetHitProxy( new HGeomPolyProxy(GeomPoly->GetParentObject(),PolyIdx) );
 			{
 				FDynamicMeshBuilder MeshBuilder;
 
@@ -561,7 +551,7 @@ void FEdModeGeometry::RenderEdge( const FSceneView* View, FPrimitiveDrawInterfac
 {
 	for( int32 ObjectIdx = 0 ; ObjectIdx < GeomObjects.Num() ; ++ObjectIdx )
 	{
-		const FGeomObject* GeometryObject = GeomObjects[ObjectIdx];
+		FGeomObjectPtr GeometryObject = GeomObjects[ObjectIdx];
 		const FColor WireColor = GeometryObject->GetActualBrush()->GetWireColor();
 
 		// Edges
@@ -570,7 +560,7 @@ void FEdModeGeometry::RenderEdge( const FSceneView* View, FPrimitiveDrawInterfac
 			const FGeomEdge* GeometryEdge = &GeometryObject->EdgePool[EdgeIdx];
 			const FColor Color = GeometryEdge->IsSelected() ? FColor(255,128,64) : WireColor;
 
-			PDI->SetHitProxy( new HGeomEdgeProxy(const_cast<FGeomObject*>(GeometryObject),EdgeIdx) );
+			PDI->SetHitProxy( new HGeomEdgeProxy(GeometryObject,EdgeIdx) );
 			{
 				FVector V0 = GeometryObject->VertexPool[ GeometryEdge->VertexIndices[0] ];
 				FVector V1 = GeometryObject->VertexPool[ GeometryEdge->VertexIndices[1] ];
@@ -592,8 +582,8 @@ void FEdModeGeometry::RenderVertex( const FSceneView* View, FPrimitiveDrawInterf
 {
 	for( int32 ObjectIdx = 0 ; ObjectIdx < GeomObjects.Num() ; ++ObjectIdx )
 	{
-		const FGeomObject* GeomObject = GeomObjects[ObjectIdx];
-		check(GeomObject);
+		FGeomObjectPtr GeomObject = GeomObjects[ObjectIdx];
+		check(GeomObject.IsValid());
 
 		// Vertices
 
@@ -611,7 +601,7 @@ void FEdModeGeometry::RenderVertex( const FSceneView* View, FPrimitiveDrawInterf
 			Scale = View->WorldToScreen( Location ).W * ( 4.0f / View->ViewRect.Width() / View->ViewMatrices.ProjMatrix.M[0][0] );
 			Color = GeomVertex->IsSelected() ? FColor(255,128,64) : GeomObject->GetActualBrush()->GetWireColor();
 
-			PDI->SetHitProxy( new HGeomVertexProxy( const_cast<FGeomObject*>(GeomObject), VertIdx) );
+			PDI->SetHitProxy( new HGeomVertexProxy( GeomObject, VertIdx) );
 			PDI->DrawSprite( Location, 4.f * Scale, 4.f * Scale, GEngine->DefaultBSPVertexTexture->Resource, Color, SDPG_Foreground, 0.0, 0.0, 0.0, 0.0 );
 			PDI->SetHitProxy( NULL );
 		}
@@ -851,8 +841,8 @@ void FEdModeGeometry::GetFromSource()
 	// Go through each brush and update its components before updating below
 	for( int32 i=0; i<GeomObjects.Num(); i++ )
 	{
-		FGeomObject* GeomObject = GeomObjects[i];
-		if(GeomObject && GeomObject->ActualBrush)
+		FGeomObjectPtr GeomObject = GeomObjects[i];
+		if(GeomObject.IsValid() && GeomObject->ActualBrush)
 		{
 #ifdef BSP_RESELECT
 			// Cache any information that'll help us reselect the object after it's reconstructed
@@ -863,7 +853,6 @@ void FEdModeGeometry::GetFromSource()
 			{
 				GeomObject->ActualBrush->RegisterAllComponents();
 			}
-			delete GeomObject;
 		}		
 	}
 	GeomObjects.Empty();
@@ -882,7 +871,7 @@ void FEdModeGeometry::GetFromSource()
 		{
 			if( BrushActor->Brush != NULL )
 			{
-				FGeomObject* GeomObject			= new FGeomObject();
+				FGeomObjectPtr GeomObject			= MakeShareable( new FGeomObject() );
 				GeomObject->SetParentObjectIndex( GeomObjects.Add( GeomObject ) );
 				GeomObject->ActualBrush			= BrushActor;
 				GeomObject->GetFromSource();
@@ -919,7 +908,7 @@ void FEdModeGeometry::SendToSource()
 {
 	for( int32 o = 0 ; o < GeomObjects.Num() ; ++o )
 	{
-		FGeomObject* go = GeomObjects[o];
+		FGeomObjectPtr go = GeomObjects[o];
 
 		go->SendToSource();
 	}
@@ -931,7 +920,7 @@ bool FEdModeGeometry::FinalizeSourceData()
 
 	for( int32 o = 0 ; o < GeomObjects.Num() ; ++o )
 	{
-		FGeomObject* go = GeomObjects[o];
+		FGeomObjectPtr go = GeomObjects[o];
 
 		if( go->FinalizeSourceData() )
 		{
@@ -952,7 +941,7 @@ void FEdModeGeometry::PostUndo()
 	for( int32 o = 0 ; o < GeomObjects.Num() ; ++o )
 	{
 		int32 Idx = 0;
-		FGeomObject* go = GeomObjects[o];
+		FGeomObjectPtr go = GeomObjects[o];
 
 		ABrush* Actor = go->GetActualBrush();
 
@@ -1062,7 +1051,7 @@ bool FModeTool_GeometryModify::BoxSelect( FBox& InBox, bool InSelect )
 
 		for( FEdModeGeometry::TGeomObjectIterator Itor( mode->GeomObjectItor() ) ; Itor ; ++Itor )
 		{
-			FGeomObject* go = *Itor;
+			FGeomObjectPtr go = *Itor;
 			FTransform ActorToWorld = go->GetActualBrush()->ActorToWorld();
 
 			// Only verts for box selection
@@ -1091,7 +1080,7 @@ bool FModeTool_GeometryModify::FrustumSelect( const FConvexVolume& InFrustum, bo
 
 		for( FEdModeGeometry::TGeomObjectIterator Itor( mode->GeomObjectItor() ) ; Itor ; ++Itor )
 		{
-			FGeomObject* go = *Itor;
+			FGeomObjectPtr go = *Itor;
 			FTransform ActorToWorld = go->GetActualBrush()->ActorToWorld();
 			// Check each vertex to see if its inside the frustum
 			for( int32 v = 0 ; v < go->VertexPool.Num() ; ++v )
@@ -1173,7 +1162,7 @@ bool FModeTool_GeometryModify::EndModify()
 		// Update internals.
 		for( FEdModeGeometry::TGeomObjectIterator Itor( mode->GeomObjectItor() ) ; Itor ; ++Itor )
 		{
-			FGeomObject* go = *Itor;
+			FGeomObjectPtr go = *Itor;
 			go->ComputeData();
 			FBSPOps::bspUnlinkPolys( go->GetActualBrush()->Brush );			
 			
@@ -1234,7 +1223,7 @@ bool FModeTool_GeometryModify::InputKey(FEditorViewportClient* ViewportClient, F
 
 		for( FEdModeGeometry::TGeomObjectIterator Itor( mode->GeomObjectItor() ) ; Itor ; ++Itor )
 		{
-			FGeomObject* go = *Itor;
+			FGeomObjectPtr go = *Itor;
 
 			for( int32 p = 0 ; p < go->PolyPool.Num() ; ++p )
 			{
diff --git a/Engine/Source/Editor/GeometryMode/Private/GeometryModifiers.cpp b/Engine/Source/Editor/GeometryMode/Private/GeometryModifiers.cpp
--- a/Engine/Source/Editor/GeometryMode/Private/GeometryModifiers.cpp
+++ b/Engine/Source/Editor/GeometryMode/Private/GeometryModifiers.cpp
@@ -196,7 +196,7 @@ bool UGeomModifier::DoEdgesOverlap()
 	//Loop through all of the geometry objects
 	for( FEdModeGeometry::TGeomObjectIterator itor( mode->GeomObjectItor() ) ; itor ; ++itor )
 	{
-		FGeomObject* geomObject = *itor;
+		FGeomObjectPtr geomObject = *itor;
 
 		//Loop through all of the edges
 		for( int32 edgeIndex1 = 0 ; edgeIndex1 < geomObject->EdgePool.Num() ; ++edgeIndex1 )
@@ -300,7 +300,7 @@ void UGeomModifier::StartTrans()
 	FEdModeGeometry* CurMode = static_cast<FEdModeGeometry*>( GLevelEditorModeTools().GetActiveMode(FBuiltinEditorModes::EM_Geometry) );
 	for( FEdModeGeometry::TGeomObjectIterator Itor( CurMode->GeomObjectItor() ) ; Itor ; ++Itor )
 	{
-		FGeomObject* go = *Itor;
+		FGeomObjectPtr go = *Itor;
 		ABrush* Actor = go->GetActualBrush();
 
 		Actor->Modify();
@@ -314,7 +314,7 @@ void UGeomModifier::EndTrans()
 }
 
 
-void UGeomModifier::StoreCurrentGeomSelections( TArray<struct FGeomSelection>& SelectionArray, FGeomObject* go ) 
+void UGeomModifier::StoreCurrentGeomSelections( TArray<struct FGeomSelection>& SelectionArray, FGeomObjectPtr go ) 
 {
 	SelectionArray.Empty();
 	FGeomSelection* gs = NULL;
@@ -366,7 +366,7 @@ void UGeomModifier::StoreAllCurrentGeomSelections()
 	// Record the current selection list into the selected brushes.
 	for( FEdModeGeometry::TGeomObjectIterator Itor( CurMode->GeomObjectItor() ) ; Itor ; ++Itor )
 	{
-		FGeomObject* go = *Itor;
+		FGeomObjectPtr go = *Itor;
 		go->CompileSelectionOrder();
 
 		ABrush* Actor = go->GetActualBrush();
@@ -413,7 +413,7 @@ bool UGeomModifier_Edit::InputDelta(FEditorViewportClient* InViewportClient,FVie
 
 	for( FEdModeGeometry::TGeomObjectIterator Itor( mode->GeomObjectItor() ) ; Itor ; ++Itor )
 	{
-		FGeomObject* go = *Itor;
+		FGeomObjectPtr go = *Itor;
 
 		for( int32 p = 0 ; p < go->PolyPool.Num() ; ++p )
 		{
@@ -722,7 +722,7 @@ void UGeomModifier_Extrude::Apply(int32 InLength, int32 InSegments)
 
 	for( FEdModeGeometry::TGeomObjectIterator Itor( mode->GeomObjectItor() ) ; Itor ; ++Itor )
 	{
-		FGeomObject* go = *Itor;
+		FGeomObjectPtr go = *Itor;
 		ABrush* Brush = go->GetActualBrush();
 
 		go->SendToSource();
@@ -2349,7 +2349,7 @@ bool UGeomModifier_Delete::OnApply()
 
 	for( FEdModeGeometry::TGeomObjectIterator Itor( mode->GeomObjectItor() ) ; Itor ; ++Itor )
 	{
-		FGeomObject* go = *Itor;
+		FGeomObjectPtr go = *Itor;
 
 		// Polys
 
@@ -2432,7 +2432,7 @@ bool UGeomModifier_Create::OnApply()
 
 	for( FEdModeGeometry::TGeomObjectIterator Itor( mode->GeomObjectItor() ) ; Itor ; ++Itor )
 	{
-		FGeomObject* go = *Itor;
+		FGeomObjectPtr go = *Itor;
 
 		go->CompileSelectionOrder();
 
@@ -2500,7 +2500,7 @@ bool UGeomModifier_Flip::OnApply()
 
 	for( FEdModeGeometry::TGeomObjectIterator Itor( mode->GeomObjectItor() ) ; Itor ; ++Itor )
 	{
-		FGeomObject* go = *Itor;
+		FGeomObjectPtr go = *Itor;
 
 		for( int32 p = 0 ; p < go->PolyPool.Num() ; ++p )
 		{
@@ -2563,14 +2563,14 @@ bool UGeomModifier_Split::OnApply()
 
 	// Get a pointer to the selected geom object
 
-	FGeomObject* GeomObject = NULL;
+	FGeomObjectPtr GeomObject = NULL;
 	for( FEdModeGeometry::TGeomObjectIterator Itor( mode->GeomObjectItor() ) ; Itor ; ++Itor )
 	{
 		GeomObject = *Itor;
 		break;
 	}
 
-	if( GeomObject == NULL )
+	if( !GeomObject.IsValid() )
 	{
 		return false;
 	}
@@ -2917,7 +2917,7 @@ bool UGeomModifier_Triangulate::OnApply()
 
 	for( FEdModeGeometry::TGeomObjectIterator Itor( mode->GeomObjectItor() ) ; Itor ; ++Itor )
 	{
-		FGeomObject* go = *Itor;
+		FGeomObjectPtr go = *Itor;
 
 		TArray<FPoly> PolyList;
 
@@ -2990,7 +2990,7 @@ bool UGeomModifier_Optimize::OnApply()
 	{
 		for( FEdModeGeometry::TGeomObjectIterator Itor( mode->GeomObjectItor() ) ; Itor ; ++Itor )
 		{
-			FGeomObject* go = *Itor;
+			FGeomObjectPtr go = *Itor;
 			ABrush* ActualBrush = go->GetActualBrush();
 
 			// Gather a list of polygons that are 
@@ -3036,7 +3036,7 @@ bool UGeomModifier_Optimize::OnApply()
 	{
 		for( FEdModeGeometry::TGeomObjectIterator Itor( mode->GeomObjectItor() ) ; Itor ; ++Itor )
 		{
-			FGeomObject* go = *Itor;
+			FGeomObjectPtr go = *Itor;
 			ABrush* ActualBrush = go->GetActualBrush();
 
 			// Optimize the polygons
@@ -3075,7 +3075,7 @@ bool UGeomModifier_Turn::OnApply()
 
 	for( FEdModeGeometry::TGeomObjectIterator Itor( mode->GeomObjectItor() ) ; Itor ; ++Itor )
 	{
-		FGeomObject* go = *Itor;
+		FGeomObjectPtr go = *Itor;
 
 		TArray<FGeomEdge> Edges;
 		go->CompileUniqueEdgeArray( &Edges );
@@ -3243,7 +3243,7 @@ bool UGeomModifier_Weld::OnApply()
 
 	for( FEdModeGeometry::TGeomObjectIterator Itor( mode->GeomObjectItor() ) ; Itor ; ++Itor )
 	{
-		FGeomObject* go = *Itor;
+		FGeomObjectPtr go = *Itor;
 
 		go->CompileSelectionOrder();
 
@@ -3279,7 +3279,7 @@ bool UGeomModifier_Weld::OnApply()
 	//finally, cache the selections AFTER the weld and set the widget to the appropriate selection
 	for( FEdModeGeometry::TGeomObjectIterator Itor( mode->GeomObjectItor() ) ; Itor ; ++Itor )
 	{
-		FGeomObject* go = *Itor;
+		FGeomObjectPtr go = *Itor;
 		go->CompileSelectionOrder();
 
 		ABrush* Actor = go->GetActualBrush();
diff --git a/Engine/Source/Editor/GeometryMode/Public/EditorGeometry.h b/Engine/Source/Editor/GeometryMode/Public/EditorGeometry.h
--- a/Engine/Source/Editor/GeometryMode/Public/EditorGeometry.h
+++ b/Engine/Source/Editor/GeometryMode/Public/EditorGeometry.h
@@ -3,9 +3,8 @@
 #pragma once
 
 
-class FGeomObject;
 class FPoly;
-
+typedef TSharedPtr<class FGeomObject> FGeomObjectPtr;
 
 /**
  * Base class for all classes that support storing and editing of geometry.
@@ -37,8 +36,8 @@ public:
 	virtual FRotator GetWidgetRotation() const			{ return Normal.Rotation(); }
 
 	//@{
-	FGeomObject* GetParentObject();
-	const FGeomObject* GetParentObject() const;
+	FGeomObjectPtr GetParentObject();
+	const FGeomObjectPtr GetParentObject() const;
 	//@}
 
 	/** Returns true if this geometry objects is a vertex. */
@@ -324,15 +323,17 @@ struct HGeomPolyProxy : public HHitProxy
 {
 	DECLARE_HIT_PROXY(UNREALED_API);
 
-	FGeomObject*	GeomObject;
+	TWeakPtr<FGeomObject> GeomObjectWeakPtr;
 	int32				PolyIndex;
 
-	HGeomPolyProxy(FGeomObject* InGeomObject, int32 InPolyIndex):
+	HGeomPolyProxy(FGeomObjectPtr InGeomObject, int32 InPolyIndex):
 		HHitProxy(HPP_UI),
-		GeomObject(InGeomObject),
+		GeomObjectWeakPtr(InGeomObject),
 		PolyIndex(InPolyIndex)
 	{}
 
+	FGeomObject* GetGeomObject() { return GeomObjectWeakPtr.IsValid() ? GeomObjectWeakPtr.Pin().Get() : nullptr; }
+
 	virtual EMouseCursor::Type GetMouseCursor() override
 	{
 		return EMouseCursor::Crosshairs;
@@ -359,15 +360,18 @@ struct HGeomEdgeProxy : public HHitProxy
 {
 	DECLARE_HIT_PROXY(UNREALED_API);
 
-	FGeomObject*	GeomObject;
+	// Kept around for backwards compat
+	TWeakPtr<FGeomObject> GeomObjectWeakPtr;
 	int32				EdgeIndex;
 
-	HGeomEdgeProxy(FGeomObject* InGeomObject, int32 InEdgeIndex):
+	HGeomEdgeProxy(FGeomObjectPtr InGeomObject, int32 InEdgeIndex):
 		HHitProxy(HPP_UI),
-		GeomObject(InGeomObject),
+		GeomObjectWeakPtr(InGeomObject),
 		EdgeIndex(InEdgeIndex)
 	{}
 
+	FGeomObject* GetGeomObject() { return GeomObjectWeakPtr.IsValid() ? GeomObjectWeakPtr.Pin().Get() : nullptr; }
+
 	virtual EMouseCursor::Type GetMouseCursor() override
 	{
 		return EMouseCursor::Crosshairs;
@@ -396,15 +400,17 @@ struct HGeomVertexProxy : public HHitProxy
 {
 	DECLARE_HIT_PROXY(UNREALED_API);
 
-	FGeomObject*	GeomObject;
+	TWeakPtr<FGeomObject> GeomObjectWeakPtr;
 	int32				VertexIndex;
 
-	HGeomVertexProxy(FGeomObject* InGeomObject, int32 InVertexIndex):
+	HGeomVertexProxy(FGeomObjectPtr InGeomObject, int32 InVertexIndex):
 		HHitProxy(HPP_UI),
-		GeomObject(InGeomObject),
+		GeomObjectWeakPtr(InGeomObject),
 		VertexIndex(InVertexIndex)
 	{}
 
+	FGeomObject* GetGeomObject() { return GeomObjectWeakPtr.IsValid() ? GeomObjectWeakPtr.Pin().Get() : nullptr; }
+
 	virtual EMouseCursor::Type GetMouseCursor() override
 	{
 		return EMouseCursor::Crosshairs;
diff --git a/Engine/Source/Editor/GeometryMode/Public/GeometryEdMode.h b/Engine/Source/Editor/GeometryMode/Public/GeometryEdMode.h
--- a/Engine/Source/Editor/GeometryMode/Public/GeometryEdMode.h
+++ b/Engine/Source/Editor/GeometryMode/Public/GeometryEdMode.h
@@ -10,6 +10,7 @@ class FGeomEdge;
 class FGeomVertex;
 class FGeomBase;
 class FGeomObject;
+typedef TSharedPtr<class FGeomObject> FGeomObjectPtr;
 
 /**
  * Geometry mode module
@@ -179,15 +180,15 @@ public:
 
 	/** @name GeomObject iterators */
 	//@{
-	typedef TArray<FGeomObject*>::TIterator TGeomObjectIterator;
-	typedef TArray<FGeomObject*>::TConstIterator TGeomObjectConstIterator;
+	typedef TArray<FGeomObjectPtr>::TIterator TGeomObjectIterator;
+	typedef TArray<FGeomObjectPtr>::TConstIterator TGeomObjectConstIterator;
 
 	TGeomObjectIterator			GeomObjectItor()			{ return TGeomObjectIterator( GeomObjects ); }
 	TGeomObjectConstIterator	GeomObjectConstItor() const	{ return TGeomObjectConstIterator( GeomObjects ); }
 
 	// @todo DB: Get rid of these; requires changes to FGeomBase::ParentObjectIndex
-	FGeomObject* GetGeomObject(int32 Index)					{ return GeomObjects[ Index ]; }
-	const FGeomObject* GetGeomObject(int32 Index) const		{ return GeomObjects[ Index ]; }
+	FGeomObjectPtr GetGeomObject(int32 Index)					{ return GeomObjects[ Index ]; }
+	const FGeomObjectPtr GetGeomObject(int32 Index) const		{ return GeomObjects[ Index ]; }
 	//@}
 
 protected:
@@ -197,7 +198,7 @@ protected:
 	 * interacts with while in this mode.  Changes done here are
 	 * reflected back to the real data in the level at specific times.
 	 */
-	TArray<FGeomObject*> GeomObjects;
+	TArray<FGeomObjectPtr> GeomObjects;
 };
 
 
diff --git a/Engine/Source/Editor/UnrealEd/Private/EditorActor.cpp b/Engine/Source/Editor/UnrealEd/Private/EditorActor.cpp
--- a/Engine/Source/Editor/UnrealEd/Private/EditorActor.cpp
+++ b/Engine/Source/Editor/UnrealEd/Private/EditorActor.cpp
@@ -2292,7 +2292,7 @@ void UUnrealEdEngine::edactAlignVertices()
 						FEdModeGeometry::TGeomObjectIterator GeomModeIt = GeomMode->GeomObjectItor();
 						for( ; GeomModeIt; ++GeomModeIt )
 						{
-							FGeomObject* Object = *GeomModeIt;
+							FGeomObjectPtr Object = *GeomModeIt;
 							if( Object->GetActualBrush() == Brush )
 							{
 								// We found our current brush, update the geometry object's data
diff --git a/Engine/Source/Editor/UnrealEd/Private/LevelEditorViewport.cpp b/Engine/Source/Editor/UnrealEd/Private/LevelEditorViewport.cpp
--- a/Engine/Source/Editor/UnrealEd/Private/LevelEditorViewport.cpp
+++ b/Engine/Source/Editor/UnrealEd/Private/LevelEditorViewport.cpp
@@ -1872,29 +1872,39 @@ void FLevelEditorViewportClient::ProcessClick(FSceneView& View, HHitProxy* HitPr
 		}
 		else if (HitProxy->IsA(HGeomPolyProxy::StaticGetType()))
 		{
-			FHitResult CheckResult(ForceInit);
-			FCollisionQueryParams BoxParams(ProcessClickTrace, false, ((HGeomPolyProxy*)HitProxy)->GeomObject->ActualBrush);
-			bool bHit = GWorld->SweepSingleByObjectType(CheckResult, Click.GetOrigin(), Click.GetOrigin() + Click.GetDirection() * HALF_WORLD_MAX, FQuat::Identity,FCollisionObjectQueryParams(ECC_WorldStatic), FCollisionShape::MakeBox(FVector(1.f)), BoxParams);
+			HGeomPolyProxy* GeomHitProxy = (HGeomPolyProxy*)HitProxy;
 
-			if( bHit )
+			if( GeomHitProxy->GetGeomObject() )
 			{
-				GEditor->UnsnappedClickLocation = CheckResult.Location;
-				GEditor->ClickLocation = CheckResult.Location;
-				GEditor->ClickPlane = FPlane(CheckResult.Location,CheckResult.Normal);
-			}
+				FHitResult CheckResult(ForceInit);
+				FCollisionQueryParams BoxParams(ProcessClickTrace, false, GeomHitProxy->GetGeomObject()->ActualBrush);
+				bool bHit = GWorld->SweepSingleByObjectType(CheckResult, Click.GetOrigin(), Click.GetOrigin() + Click.GetDirection() * HALF_WORLD_MAX, FQuat::Identity, FCollisionObjectQueryParams(ECC_WorldStatic), FCollisionShape::MakeBox(FVector(1.f)), BoxParams);
 
-			if( !ClickHandlers::ClickActor(this,((HGeomPolyProxy*)HitProxy)->GeomObject->ActualBrush,Click,false) )
-			{
-				ClickHandlers::ClickGeomPoly(this,(HGeomPolyProxy*)HitProxy,Click);
-			}
+				if(bHit)
+				{
+					GEditor->UnsnappedClickLocation = CheckResult.Location;
+					GEditor->ClickLocation = CheckResult.Location;
+					GEditor->ClickPlane = FPlane(CheckResult.Location, CheckResult.Normal);
+				}
+
+				if(!ClickHandlers::ClickActor(this, GeomHitProxy->GetGeomObject()->ActualBrush, Click, false))
+				{
+					ClickHandlers::ClickGeomPoly(this, GeomHitProxy, Click);
+				}
 
-			Invalidate( true, true );
+				Invalidate(true, true);
+			}
 		}
 		else if (HitProxy->IsA(HGeomEdgeProxy::StaticGetType()))
 		{
-			if( !ClickHandlers::ClickGeomEdge(this,(HGeomEdgeProxy*)HitProxy,Click) )
+			HGeomEdgeProxy* GeomHitProxy = (HGeomEdgeProxy*)HitProxy;
+
+			if( GeomHitProxy->GetGeomObject() !=nullptr )
 			{
-				ClickHandlers::ClickActor(this,((HGeomEdgeProxy*)HitProxy)->GeomObject->ActualBrush,Click,true);
+				if(!ClickHandlers::ClickGeomEdge(this, GeomHitProxy, Click))
+				{
+					ClickHandlers::ClickActor(this, GeomHitProxy->GetGeomObject()->ActualBrush, Click, true);
+				}
 			}
 		}
 		else if (HitProxy->IsA(HGeomVertexProxy::StaticGetType()))
diff --git a/Engine/Source/Editor/UnrealEd/Private/LevelViewportClickHandlers.cpp b/Engine/Source/Editor/UnrealEd/Private/LevelViewportClickHandlers.cpp
--- a/Engine/Source/Editor/UnrealEd/Private/LevelViewportClickHandlers.cpp
+++ b/Engine/Source/Editor/UnrealEd/Private/LevelViewportClickHandlers.cpp
@@ -447,6 +447,11 @@ namespace ClickHandlers
 
 	bool ClickGeomPoly(FLevelEditorViewportClient* ViewportClient, HGeomPolyProxy* InHitProxy, const FViewportClick& Click)
 	{
+		if( !InHitProxy->GeomObjectWeakPtr.IsValid() )
+		{
+			return false;
+		}
+
 		// Pivot snapping
 		if( Click.GetKey() == EKeys::MiddleMouseButton && Click.IsAltDown() )
 		{
@@ -456,7 +461,7 @@ namespace ClickHandlers
 		}
 		else if( Click.GetKey() == EKeys::LeftMouseButton && Click.IsControlDown() && Click.IsShiftDown() && !Click.IsAltDown() )
 		{
-			GEditor->SelectActor( InHitProxy->GeomObject->GetActualBrush(), false, true );
+			GEditor->SelectActor( InHitProxy->GetGeomObject()->GetActualBrush(), false, true );
 		}
 		else if( Click.GetKey() == EKeys::LeftMouseButton )
 		{
@@ -464,7 +469,7 @@ namespace ClickHandlers
 			FEdMode* Mode = GLevelEditorModeTools().GetActiveMode( FBuiltinEditorModes::EM_Geometry );
 			if( Mode )
 			{
-				if( ( InHitProxy != NULL) && ( InHitProxy->GeomObject != NULL ) && (InHitProxy->GeomObject->PolyPool.IsValidIndex( InHitProxy->PolyIndex ) == true ) )					
+				if( ( InHitProxy != NULL) && ( InHitProxy->GetGeomObject() != NULL ) && (InHitProxy->GetGeomObject()->PolyPool.IsValidIndex( InHitProxy->PolyIndex ) == true ) )					
 				{
 					Mode->GetCurrentTool()->StartTrans();
 
@@ -473,7 +478,7 @@ namespace ClickHandlers
 						Mode->GetCurrentTool()->SelectNone();
 					}
 
-					FGeomPoly& gp = InHitProxy->GeomObject->PolyPool[ InHitProxy->PolyIndex ];
+					FGeomPoly& gp = InHitProxy->GetGeomObject()->PolyPool[ InHitProxy->PolyIndex ];
 					gp.Select( Click.IsControlDown() ? !gp.IsSelected() : true );
 
 					Mode->SelectionChanged();
@@ -488,17 +493,13 @@ namespace ClickHandlers
 					{
 						UE_LOG(LogEditorViewport, Warning, TEXT("Invalid hitproxy" ) );
 					}
-					else if ( InHitProxy->GeomObject == NULL ) 
-					{
-						UE_LOG(LogEditorViewport, Warning, TEXT("Invalid GeomObject" ) );		
-					}
 					else
 					{
 						//try to get the name of the object also
 						FString name = TEXT("UNKNOWN");
-						if( InHitProxy->GeomObject->GetActualBrush() != NULL )
+						if( InHitProxy->GetGeomObject()->GetActualBrush() != NULL )
 						{
-							name = InHitProxy->GeomObject->GetActualBrush()->GetName();
+							name = InHitProxy->GetGeomObject()->GetActualBrush()->GetName();
 						}
 						UE_LOG(LogEditorViewport, Warning, TEXT("Invalid PolyIndex %d on %s" ) ,InHitProxy->PolyIndex , *name );
 					}
@@ -540,6 +541,11 @@ namespace ClickHandlers
 
 	bool ClickGeomEdge(FLevelEditorViewportClient* ViewportClient, HGeomEdgeProxy* InHitProxy, const FViewportClick& Click)
 	{
+		if( !InHitProxy->GetGeomObject() )
+		{
+			return false;
+		}
+
 		// Pivot snapping
 		if( Click.GetKey() == EKeys::MiddleMouseButton && Click.IsAltDown() )
 		{
@@ -549,7 +555,7 @@ namespace ClickHandlers
 		}
 		else if( Click.GetKey() == EKeys::LeftMouseButton && Click.IsControlDown() && Click.IsShiftDown() && !Click.IsAltDown() )
 		{
-			GEditor->SelectActor( InHitProxy->GeomObject->GetActualBrush(), false, true );
+			GEditor->SelectActor( InHitProxy->GetGeomObject()->GetActualBrush(), false, true );
 
 			return true;
 		}
@@ -567,17 +573,17 @@ namespace ClickHandlers
 					Mode->GetCurrentTool()->SelectNone();
 				}
 
-				FGeomEdge& HitEdge = InHitProxy->GeomObject->EdgePool[ InHitProxy->EdgeIndex ];
+				FGeomEdge& HitEdge = InHitProxy->GetGeomObject()->EdgePool[ InHitProxy->EdgeIndex ];
 				HitEdge.Select( bControlDown ? !HitEdge.IsSelected() : true );
 
 				if( ViewportClient->IsOrtho() )
 				{
 					// Select all edges in the brush that match the projected mid point of the original edge.
-					for( int32 EdgeIndex = 0 ; EdgeIndex < InHitProxy->GeomObject->EdgePool.Num() ; ++EdgeIndex )
+					for( int32 EdgeIndex = 0 ; EdgeIndex < InHitProxy->GetGeomObject()->EdgePool.Num() ; ++EdgeIndex )
 					{
 						if ( EdgeIndex != InHitProxy->EdgeIndex )
 						{
-							FGeomEdge& GeomEdge = InHitProxy->GeomObject->EdgePool[ EdgeIndex ];
+							FGeomEdge& GeomEdge = InHitProxy->GetGeomObject()->EdgePool[ EdgeIndex ];
 							if ( OrthoEqual( ViewportClient->ViewportType, GeomEdge.GetMid(), HitEdge.GetMid() ) )
 							{
 								GeomEdge.Select( bControlDown ? !GeomEdge.IsSelected() : true );
@@ -603,6 +609,11 @@ namespace ClickHandlers
 
 	bool ClickGeomVertex(FLevelEditorViewportClient* ViewportClient,HGeomVertexProxy* InHitProxy,const FViewportClick& Click)
 	{
+		if(InHitProxy->GetGeomObject() == nullptr)
+		{
+			return false;
+		}
+
 		if( !GLevelEditorModeTools().IsModeActive( FBuiltinEditorModes::EM_Geometry ) )
 		{
 			return false;
@@ -616,14 +627,8 @@ namespace ClickHandlers
 
 		if( Click.GetKey() == EKeys::RightMouseButton )
 		{
-			// sanity checks & defensive coding seem necessary based on Prio 1 TTP #318426
-			ensure( NULL != InHitProxy->GeomObject );
-			if ( InHitProxy->GeomObject == NULL ) 
-			{
-				UE_LOG(LogEditorViewport, Warning, TEXT("Invalid GeomObject" ) );	
-				return false;
-			}
-			if (InHitProxy->VertexIndex < 0 || InHitProxy->VertexIndex >= InHitProxy->GeomObject->VertexPool.Num() )
+
+			if (InHitProxy->VertexIndex < 0 || InHitProxy->VertexIndex >= InHitProxy->GetGeomObject()->VertexPool.Num() )
 			{
 				UE_LOG(LogEditorViewport, Warning, TEXT("Invalid InHitProxy->VertexIndex" ) );		
 				return false;
@@ -633,7 +638,7 @@ namespace ClickHandlers
 			Tool->StartTrans();
 
 			// Compute out far to move to get back on the grid.
-			const FVector WorldLoc = InHitProxy->GeomObject->GetActualBrush()->ActorToWorld().TransformPosition( InHitProxy->GeomObject->VertexPool[ InHitProxy->VertexIndex ] );
+			const FVector WorldLoc = InHitProxy->GetGeomObject()->GetActualBrush()->ActorToWorld().TransformPosition( InHitProxy->GetGeomObject()->VertexPool[ InHitProxy->VertexIndex ] );
 
 			FVector SnappedLoc = WorldLoc;
 			FSnappingUtils::SnapPointToGrid( SnappedLoc, FVector(GEditor->GetGridSize()) );
@@ -641,9 +646,9 @@ namespace ClickHandlers
 			const FVector Delta = SnappedLoc - WorldLoc;
 			GEditor->SetPivot( SnappedLoc, false, false );
 
-			for( int32 VertexIndex = 0 ; VertexIndex < InHitProxy->GeomObject->VertexPool.Num() ; ++VertexIndex )
+			for( int32 VertexIndex = 0 ; VertexIndex < InHitProxy->GetGeomObject()->VertexPool.Num() ; ++VertexIndex )
 			{
-				FGeomVertex& GeomVertex = InHitProxy->GeomObject->VertexPool[VertexIndex];
+				FGeomVertex& GeomVertex = InHitProxy->GetGeomObject()->VertexPool[VertexIndex];
 				if( GeomVertex.IsSelected() )
 				{
 					GeomVertex += Delta;
@@ -651,7 +656,7 @@ namespace ClickHandlers
 			}
 
 			Tool->EndTrans();
-			InHitProxy->GeomObject->SendToSource();
+			InHitProxy->GetGeomObject()->SendToSource();
 			ViewportClient->Invalidate( true, true );
 
 			// HACK: The Bsp update has to occur after SendToSource() updates the vert pool, putting it outside
@@ -674,7 +679,7 @@ namespace ClickHandlers
 		}
 		else if( Click.GetKey() == EKeys::LeftMouseButton && Click.IsControlDown() && Click.IsShiftDown() && !Click.IsAltDown() )
 		{
-			GEditor->SelectActor( InHitProxy->GeomObject->GetActualBrush(), false, true );
+			GEditor->SelectActor( InHitProxy->GetGeomObject()->GetActualBrush(), false, true );
 		}
 		else if( Click.GetKey() == EKeys::LeftMouseButton )
 		{
@@ -687,7 +692,7 @@ namespace ClickHandlers
 				Mode->GetCurrentTool()->SelectNone();
 			}
 
-			FGeomVertex& HitVertex = InHitProxy->GeomObject->VertexPool[ InHitProxy->VertexIndex ];
+			FGeomVertex& HitVertex = InHitProxy->GetGeomObject()->VertexPool[ InHitProxy->VertexIndex ];
 			bool bSelect = bControlDown ? !HitVertex.IsSelected() : true;
 
 			HitVertex.Select( bSelect );
@@ -695,11 +700,11 @@ namespace ClickHandlers
 			if( ViewportClient->IsOrtho() )
 			{
 				// Select all vertices that project to the same location.
-				for( int32 VertexIndex = 0 ; VertexIndex < InHitProxy->GeomObject->VertexPool.Num() ; ++VertexIndex )
+				for( int32 VertexIndex = 0 ; VertexIndex < InHitProxy->GetGeomObject()->VertexPool.Num() ; ++VertexIndex )
 				{
 					if ( VertexIndex != InHitProxy->VertexIndex )
 					{
-						FGeomVertex& GeomVertex = InHitProxy->GeomObject->VertexPool[VertexIndex];
+						FGeomVertex& GeomVertex = InHitProxy->GetGeomObject()->VertexPool[VertexIndex];
 						if ( OrthoEqual( ViewportClient->ViewportType, GeomVertex, HitVertex ) )
 						{
 							GeomVertex.Select( bSelect );
