Fixed undo in Geometry Mode so that the geometry selection state is correctly restored.

diff --git a/Engine/Source/Editor/GeometryMode/Private/EditorGeometry.cpp b/Engine/Source/Editor/GeometryMode/Private/EditorGeometry.cpp
--- a/Engine/Source/Editor/GeometryMode/Private/EditorGeometry.cpp
+++ b/Engine/Source/Editor/GeometryMode/Private/EditorGeometry.cpp
@@ -102,6 +102,32 @@ int32 FGeomObject::SetPivotFromSelectionArray( TArray<struct FGeomSelection>& Se
 	return HighestSelectionIndex;
 }
 
+void FGeomObject::UpdateFromSelectionArray(TArray<struct FGeomSelection>& SelectionArray)
+{
+	FEditorModeTools& Tools = GLevelEditorModeTools();
+
+	for (const auto& Selection : SelectionArray)
+	{
+		switch (Selection.Type)
+		{
+		case GS_Poly:
+			PolyPool[Selection.Index].ForceSelectionIndex(Selection.SelectionIndex);
+			break;
+
+		case GS_Edge:
+			EdgePool[Selection.Index].ForceSelectionIndex(Selection.SelectionIndex);
+			break;
+
+		case GS_Vertex:
+			VertexPool[Selection.Index].ForceSelectionIndex(Selection.SelectionIndex);
+			break;
+		}
+	}
+
+	DirtySelectionOrder();
+	CompileSelectionOrder();
+}
+
 ///////////////////////////////////////////////////////////////////////////////
 //
 // FGeomVertex
diff --git a/Engine/Source/Editor/GeometryMode/Private/GeometryEdMode.cpp b/Engine/Source/Editor/GeometryMode/Private/GeometryEdMode.cpp
--- a/Engine/Source/Editor/GeometryMode/Private/GeometryEdMode.cpp
+++ b/Engine/Source/Editor/GeometryMode/Private/GeometryEdMode.cpp
@@ -949,6 +949,7 @@ void FEdModeGeometry::PostUndo()
 		go->SelectNone();
 
 		// Next, restore the cached selection
+		go->UpdateFromSelectionArray( Actor->SavedSelections );
 		int32 res = go->SetPivotFromSelectionArray( Actor->SavedSelections );
 		
 		//use the centre of the actor if we didnt find a suitable selection
diff --git a/Engine/Source/Editor/GeometryMode/Public/EditorGeometry.h b/Engine/Source/Editor/GeometryMode/Public/EditorGeometry.h
--- a/Engine/Source/Editor/GeometryMode/Public/EditorGeometry.h
+++ b/Engine/Source/Editor/GeometryMode/Public/EditorGeometry.h
@@ -305,6 +305,11 @@ public:
 	 */
 	int32 SetPivotFromSelectionArray( TArray<struct FGeomSelection>& SelectionArray );
 
+	/**
+	 * Update the selection state based on the passed in array
+	 */
+	void UpdateFromSelectionArray( TArray<struct FGeomSelection>& SelectionArray );
+
 private:
 	int32 LastSelectionIndex;
 };
