Geometry Optimize operation first performs a triangulation so that it works correctly.

diff --git a/Engine/Source/Editor/GeometryMode/Classes/GeomModifier_Optimize.h b/Engine/Source/Editor/GeometryMode/Classes/GeomModifier_Optimize.h
--- a/Engine/Source/Editor/GeometryMode/Classes/GeomModifier_Optimize.h
+++ b/Engine/Source/Editor/GeometryMode/Classes/GeomModifier_Optimize.h
@@ -9,13 +9,12 @@
 #include "GeomModifier_Optimize.generated.h"
 
 UCLASS()
-class UGeomModifier_Optimize : public UGeomModifier_Edit
+class UGeomModifier_Optimize : public UGeomModifier_Triangulate
 {
 	GENERATED_UCLASS_BODY()
 
 
 	// Begin UGeomModifier Interface
-	virtual bool Supports() override;
 protected:
 	virtual bool OnApply() override;
 	// End UGeomModifier Interface
diff --git a/Engine/Source/Editor/GeometryMode/Private/GeometryModifiers.cpp b/Engine/Source/Editor/GeometryMode/Private/GeometryModifiers.cpp
--- a/Engine/Source/Editor/GeometryMode/Private/GeometryModifiers.cpp
+++ b/Engine/Source/Editor/GeometryMode/Private/GeometryModifiers.cpp
@@ -2979,15 +2979,11 @@ UGeomModifier_Optimize::UGeomModifier_Optimize(const FObjectInitializer& ObjectI
 	bPushButton = true;
 }
 
-bool UGeomModifier_Optimize::Supports()
-{
-	FEdModeGeometry* mode = (FEdModeGeometry*)GLevelEditorModeTools().GetActiveMode(FBuiltinEditorModes::EM_Geometry);
-	return !mode->HaveVerticesSelected() && !mode->HaveEdgesSelected();
-}
-
-
 bool UGeomModifier_Optimize::OnApply()
 {
+	// First triangulate before performing optimize
+	Super::OnApply();
+
 	FEdModeGeometry* mode = (FEdModeGeometry*)GLevelEditorModeTools().GetActiveMode(FBuiltinEditorModes::EM_Geometry);
 
 	TArray<FPoly> Polygons;
