Added facility to automatically adjust the pivot offset for Brushes when they are edited in Geometry Mode, so the pivot stays centered on the vertices.

diff --git a/Engine/Config/BaseEditorPerProjectUserSettings.ini b/Engine/Config/BaseEditorPerProjectUserSettings.ini
--- a/Engine/Config/BaseEditorPerProjectUserSettings.ini
+++ b/Engine/Config/BaseEditorPerProjectUserSettings.ini
@@ -167,6 +167,8 @@ TextDiffToolPath=(FilePath="p4merge.exe")
 bAutoApplyLightingEnable=true
 ; If true, BSP will auto-update
 bBSPAutoUpdate=True
+; If true, the Pivot Offset for BSP will automatically move to stay centered on its vertices
+bAutoMoveBSPPivotOffset=False
 ; If true, Navigation will auto-update
 bNavigationAutoUpdate=True
 ; If true, Replaces respects the scale of the original actor. If false, Replace sets the scale to 1.0
diff --git a/Engine/Source/Editor/GeometryMode/Classes/GeomModifier.h b/Engine/Source/Editor/GeometryMode/Classes/GeomModifier.h
--- a/Engine/Source/Editor/GeometryMode/Classes/GeomModifier.h
+++ b/Engine/Source/Editor/GeometryMode/Classes/GeomModifier.h
@@ -143,6 +143,11 @@ protected:
 	 * Implements the modifier application.
 	 */
 	virtual bool OnApply();
+
+	/**
+	 * Updates the pivot offset of the selected brushes based on the current vertex positions
+	 */
+	virtual void UpdatePivotOffset();
 };
 
 
diff --git a/Engine/Source/Editor/GeometryMode/Private/GeometryModifiers.cpp b/Engine/Source/Editor/GeometryMode/Private/GeometryModifiers.cpp
--- a/Engine/Source/Editor/GeometryMode/Private/GeometryModifiers.cpp
+++ b/Engine/Source/Editor/GeometryMode/Private/GeometryModifiers.cpp
@@ -246,6 +246,48 @@ bool UGeomModifier::DoEdgesOverlap()
 	return false;
 }
 
+void UGeomModifier::UpdatePivotOffset()
+{
+	if (!GetDefault<ULevelEditorMiscSettings>()->bAutoMoveBSPPivotOffset)
+	{
+		return;
+	}
+
+	FEdModeGeometry* Mode = (FEdModeGeometry*)GLevelEditorModeTools().GetActiveMode(FBuiltinEditorModes::EM_Geometry);
+
+	for (FEdModeGeometry::TGeomObjectIterator It(Mode->GeomObjectItor()); It; ++It)
+	{
+		FGeomObjectPtr GeomObject = *It;
+		ABrush* Brush = GeomObject->GetActualBrush();
+
+		TSet<FVector> UniqueVertices;
+		FVector VertexCenter = FVector::ZeroVector;
+
+		if (Brush->Brush && Brush->Brush->Polys)
+		{
+			for (const auto& Element : Brush->Brush->Polys->Element)
+			{
+				for (const auto& Vertex : Element.Vertices)
+				{
+					UniqueVertices.Add(Vertex);
+				}
+			}
+
+			for (const auto& Vertex : UniqueVertices)
+			{
+				VertexCenter += Vertex;
+			}
+
+			if (UniqueVertices.Num() > 0)
+			{
+				VertexCenter /= UniqueVertices.Num();
+			}
+		}
+
+		Brush->SetPivotOffset(Brush->GetTransform().TransformPosition(VertexCenter));
+	}
+}
+
 ///////////////////////////////////////////////////////////////////////////////////////////////////////////////////////
 //
 //	Transaction tracking.
@@ -579,6 +621,8 @@ bool UGeomModifier_Edit::InputDelta(FEditorViewportClient* InViewportClient,FVie
 
 	EndTrans();
 
+	UpdatePivotOffset();
+
 	GEditor->RedrawLevelEditingViewports(true);
 
 	return true;
@@ -662,6 +706,8 @@ bool UGeomModifier_Extrude::OnApply()
 
 	GEditor->RebuildAlteredBSP(); // Brush has been altered, update the Bsp
 
+	UpdatePivotOffset();
+
 	return true;
 }
 
@@ -862,6 +908,8 @@ bool UGeomModifier_Lathe::OnApply()
 
 	GEditor->RebuildAlteredBSP(); // Brush has been altered, update the Bsp
 
+	UpdatePivotOffset();
+
 	return true;
 }
 
@@ -1176,6 +1224,8 @@ bool UGeomModifier_Pen::OnApply()
 {
 	Apply();
 
+	UpdatePivotOffset();
+
 	return true;
 }
 
@@ -1979,6 +2029,8 @@ bool UGeomModifier_Clip::OnApply()
 
 	GEditor->RebuildAlteredBSP(); // Brush has been altered, update the Bsp
 
+	UpdatePivotOffset();
+
 	return true;
 }
 
@@ -2417,6 +2469,8 @@ bool UGeomModifier_Delete::OnApply()
 		Tools.SetPivotLocation( SelectedActor->GetActorLocation() , false );
 	}
 	
+	UpdatePivotOffset();
+
 	return bHandled;
 }
 
@@ -2482,6 +2536,7 @@ bool UGeomModifier_Create::OnApply()
 
 	GEditor->RebuildAlteredBSP(); // Brush has been altered, update the Bsp
 
+	UpdatePivotOffset();
 	return true;
 }
 
@@ -2527,6 +2582,7 @@ bool UGeomModifier_Flip::OnApply()
 
 	GEditor->RebuildAlteredBSP(); // Brush has been altered, update the Bsp
 
+	UpdatePivotOffset();
 	return true;
 }
 
@@ -2898,6 +2954,8 @@ bool UGeomModifier_Split::OnApply()
 
 	GEditor->RebuildAlteredBSP(); // Brush has been altered, update the Bsp
 
+	UpdatePivotOffset();
+
 	return true;
 }
 
diff --git a/Engine/Source/Editor/UnrealEd/Classes/Settings/LevelEditorMiscSettings.h b/Engine/Source/Editor/UnrealEd/Classes/Settings/LevelEditorMiscSettings.h
--- a/Engine/Source/Editor/UnrealEd/Classes/Settings/LevelEditorMiscSettings.h
+++ b/Engine/Source/Editor/UnrealEd/Classes/Settings/LevelEditorMiscSettings.h
@@ -26,6 +26,10 @@ public:
 	UPROPERTY(EditAnywhere, config, Category=Editing, meta=( DisplayName = "Update BSP Automatically" ))
 	uint32 bBSPAutoUpdate:1;
 
+	/** If true, the pivot offset for BSP will be automatically moved to stay centered on its vertices */
+	UPROPERTY(EditAnywhere, config, Category=Editing, meta=( DisplayName = "Move BSP Pivot Offset Automatically" ))
+	uint32 bAutoMoveBSPPivotOffset:1;
+
 	/** If true, Navigation will auto-update */
 	UPROPERTY(EditAnywhere, config, Category=Editing, meta=( DisplayName = "Update Navigation Automatically" ))
 	uint32 bNavigationAutoUpdate:1;
