Fixed issue with automatic pivot offset when editing brushes.

diff --git a/Engine/Source/Editor/GeometryMode/Classes/GeomModifier.h b/Engine/Source/Editor/GeometryMode/Classes/GeomModifier.h
--- a/Engine/Source/Editor/GeometryMode/Classes/GeomModifier.h
+++ b/Engine/Source/Editor/GeometryMode/Classes/GeomModifier.h
@@ -33,6 +33,10 @@ class UGeomModifier : public UObject
 	UPROPERTY(EditAnywhere, Category=GeomModifier)
 	uint32 bInitialized:1;
 
+	/** If true, the pivot offset should be updated when the modification ends */
+	UPROPERTY()
+	uint32 bPendingPivotOffsetUpdate:1;
+
 private:
 	/** Stored state of polys in case the brush state needs to be restroed */
 	UPROPERTY()
diff --git a/Engine/Source/Editor/GeometryMode/Private/GeometryModifiers.cpp b/Engine/Source/Editor/GeometryMode/Private/GeometryModifiers.cpp
--- a/Engine/Source/Editor/GeometryMode/Private/GeometryModifiers.cpp
+++ b/Engine/Source/Editor/GeometryMode/Private/GeometryModifiers.cpp
@@ -36,6 +36,7 @@ UGeomModifier::UGeomModifier(const FObjectInitializer& ObjectInitializer)
 {
 	bPushButton = false;
 	bInitialized = false;
+	bPendingPivotOffsetUpdate = false;
 	CachedPolys = NULL;
 }
 
@@ -109,12 +110,17 @@ void UGeomModifier::GeomError(const FString& InErrorMsg)
 bool UGeomModifier::StartModify()
 {
 	bInitialized = false;
+	bPendingPivotOffsetUpdate = false;
 	return false;
 }
 
 
 bool UGeomModifier::EndModify()
 {
+	if (bPendingPivotOffsetUpdate)
+	{
+		UpdatePivotOffset();
+	}
 	StoreAllCurrentGeomSelections();
 	FEditorSupportDelegates::RedrawAllViewports.Broadcast();
 	return true;
@@ -284,7 +290,7 @@ void UGeomModifier::UpdatePivotOffset()
 			}
 		}
 
-		Brush->SetPivotOffset(Brush->GetTransform().TransformPosition(VertexCenter));
+		Brush->SetPivotOffset(VertexCenter);
 	}
 }
 
@@ -620,9 +626,7 @@ bool UGeomModifier_Edit::InputDelta(FEditorViewportClient* InViewportClient,FVie
 	}
 
 	EndTrans();
-
-	UpdatePivotOffset();
-
+	bPendingPivotOffsetUpdate = true;
 	GEditor->RedrawLevelEditingViewports(true);
 
 	return true;
@@ -705,8 +709,7 @@ bool UGeomModifier_Extrude::OnApply()
 	GLevelEditorModeTools().SetCoordSystem(SaveCS);
 
 	GEditor->RebuildAlteredBSP(); // Brush has been altered, update the Bsp
-
-	UpdatePivotOffset();
+	bPendingPivotOffsetUpdate = true;
 
 	return true;
 }
@@ -908,8 +911,7 @@ bool UGeomModifier_Lathe::OnApply()
 
 	GEditor->RebuildAlteredBSP(); // Brush has been altered, update the Bsp
 
-	UpdatePivotOffset();
-
+	bPendingPivotOffsetUpdate = true;
 	return true;
 }
 
@@ -1224,8 +1226,7 @@ bool UGeomModifier_Pen::OnApply()
 {
 	Apply();
 
-	UpdatePivotOffset();
-
+	bPendingPivotOffsetUpdate = true;
 	return true;
 }
 
@@ -2029,8 +2030,7 @@ bool UGeomModifier_Clip::OnApply()
 
 	GEditor->RebuildAlteredBSP(); // Brush has been altered, update the Bsp
 
-	UpdatePivotOffset();
-
+	bPendingPivotOffsetUpdate = true;
 	return true;
 }
 
@@ -2469,8 +2469,7 @@ bool UGeomModifier_Delete::OnApply()
 		Tools.SetPivotLocation( SelectedActor->GetActorLocation() , false );
 	}
 	
-	UpdatePivotOffset();
-
+	bPendingPivotOffsetUpdate = true;
 	return bHandled;
 }
 
@@ -2536,7 +2535,7 @@ bool UGeomModifier_Create::OnApply()
 
 	GEditor->RebuildAlteredBSP(); // Brush has been altered, update the Bsp
 
-	UpdatePivotOffset();
+	bPendingPivotOffsetUpdate = true;
 	return true;
 }
 
@@ -2582,7 +2581,7 @@ bool UGeomModifier_Flip::OnApply()
 
 	GEditor->RebuildAlteredBSP(); // Brush has been altered, update the Bsp
 
-	UpdatePivotOffset();
+	bPendingPivotOffsetUpdate = true;
 	return true;
 }
 
@@ -2954,8 +2953,7 @@ bool UGeomModifier_Split::OnApply()
 
 	GEditor->RebuildAlteredBSP(); // Brush has been altered, update the Bsp
 
-	UpdatePivotOffset();
-
+	bPendingPivotOffsetUpdate = true;
 	return true;
 }
 
