Fixed brush extrusion when Brush actor is rotated.

diff --git a/Engine/Source/Editor/GeometryMode/Private/GeometryModifiers.cpp b/Engine/Source/Editor/GeometryMode/Private/GeometryModifiers.cpp
--- a/Engine/Source/Editor/GeometryMode/Private/GeometryModifiers.cpp
+++ b/Engine/Source/Editor/GeometryMode/Private/GeometryModifiers.cpp
@@ -796,7 +796,7 @@ void UGeomModifier_Extrude::Apply(int32 InLength, int32 InSegments)
 		{
 			FGeomPoly* gp = &go->PolyPool[p];
 
-			FVector Normal = mode->GetWidgetNormalFromCurrentAxis( gp );
+			FVector Normal = Brush->GetActorQuat().Inverse().RotateVector( mode->GetWidgetNormalFromCurrentAxis( gp ) );
 
 			if( gp->IsSelected() )
 			{
